Fix static analysis error about potentially indexing past end of array.

diff --git a/Engine/Plugins/Lumin/MagicLeap/Source/MagicLeap/Private/MagicLeapCustomPresent.cpp b/Engine/Plugins/Lumin/MagicLeap/Source/MagicLeap/Private/MagicLeapCustomPresent.cpp
--- a/Engine/Plugins/Lumin/MagicLeap/Source/MagicLeap/Private/MagicLeapCustomPresent.cpp
+++ b/Engine/Plugins/Lumin/MagicLeap/Source/MagicLeap/Private/MagicLeapCustomPresent.cpp
@@ -318,6 +318,8 @@ void FMagicLeapCustomPresentOpenGL::FinishRendering()
 			glDisable(GL_FRAMEBUFFER_SRGB);
 		}
 
+		check(vp_array.num_virtual_cameras >= 2); // We assume at least one virtual camera per eye
+
 		FVector2D InternalTextureDims;
 		Plugin->GetAppFrameworkConst().GetDeviceResolution(InternalTextureDims);
 
@@ -353,6 +355,8 @@ void FMagicLeapCustomPresentOpenGL::FinishRendering()
 			glEnable(GL_FRAMEBUFFER_SRGB);
 		}
 
+		static_assert(ARRAY_COUNT(vp_array.virtual_cameras) == 2, "The MLSDK has updated the size of the virtual_cameras array.");
+#if 0 // Enable this in case the MLSDK increases the size of the virtual_cameras array past 2
 		for (uint32 i = 2; i < vp_array.num_virtual_cameras; ++i)
 		{
 			bResult = MLGraphicsSignalSyncObjectGL(Plugin->GraphicsClient, vp_array.virtual_cameras[i].sync_object, &OutStatus);
@@ -361,6 +365,7 @@ void FMagicLeapCustomPresentOpenGL::FinishRendering()
 				UE_LOG(LogMagicLeap, Error, TEXT("MLGraphicsSignalSyncObjectGL for eye %d failed with status %d"), i, static_cast<int32>(OutStatus));
 			}
 		}
+#endif
 
 #if BEGIN_END_FRAME_BALANCE_HACK
 		--BalanceCounter;
