CIS fix - Removing needless condition.

diff --git a/Engine/Plugins/Lumin/MagicLeap/Source/MagicLeapAudio/Private/AudioMixerPlatformMagicLeap.cpp b/Engine/Plugins/Lumin/MagicLeap/Source/MagicLeapAudio/Private/AudioMixerPlatformMagicLeap.cpp
--- a/Engine/Plugins/Lumin/MagicLeap/Source/MagicLeapAudio/Private/AudioMixerPlatformMagicLeap.cpp
+++ b/Engine/Plugins/Lumin/MagicLeap/Source/MagicLeapAudio/Private/AudioMixerPlatformMagicLeap.cpp
@@ -322,21 +322,18 @@ namespace Audio
 			return false;
 		}
 
-		if (AudioStreamInfo.StreamState != EAudioOutputStreamState::Stopped)
+		if (!MLAudioStopSound(StreamHandle, &Result))
 		{
-			if (!MLAudioStopSound(StreamHandle, &Result))
-			{
-				MLAUDIO_LOG_ON_FAIL(Result);
-				return false;
-			}
-
-			if (AudioStreamInfo.StreamState == EAudioOutputStreamState::Running)
-			{
-				StopGeneratingAudio();
-			}
+			MLAUDIO_LOG_ON_FAIL(Result);
+			return false;
+		}
 
-			check(AudioStreamInfo.StreamState == EAudioOutputStreamState::Stopped);
+		if (AudioStreamInfo.StreamState == EAudioOutputStreamState::Running)
+		{
+			StopGeneratingAudio();
 		}
+
+		check(AudioStreamInfo.StreamState == EAudioOutputStreamState::Stopped);
 #endif //WITH_MLSDK
 
 		return true;
