Mirroring CL 4071045 to Main
CIS Fix - Wrapping 3rd party headers with start/end macros so we can pragma out warnings they introduce.

diff --git a/Engine/Plugins/Lumin/MagicLeap/Source/MLSDK/Public/MagicLeapMath.h b/Engine/Plugins/Lumin/MagicLeap/Source/MLSDK/Public/MagicLeapMath.h
--- a/Engine/Plugins/Lumin/MagicLeap/Source/MLSDK/Public/MagicLeapMath.h
+++ b/Engine/Plugins/Lumin/MagicLeap/Source/MLSDK/Public/MagicLeapMath.h
@@ -36,13 +36,12 @@
 
 #include "Kismet/KismetMathLibrary.h"
 #include "CoreMinimal.h"
+#include "MagicLeapPluginUtil.h" // for ML_INCLUDES_START/END
+
 #if WITH_MLSDK
-THIRD_PARTY_INCLUDES_START // pushes current warning pragmas
-#if PLATFORM_WINDOWS
-	__pragma(warning(disable: 4201)) /* warning C4201: nonstandard extension used: nameless struct/union */
-#endif
+ML_INCLUDES_START
 #include <ml_types.h>
-THIRD_PARTY_INCLUDES_END // pops current warning pragmas
+ML_INCLUDES_END
 #endif //WITH_MLSDK
 
 namespace MagicLeap
diff --git a/Engine/Plugins/Lumin/MagicLeap/Source/MLSDK/Public/MagicLeapPluginUtil.h b/Engine/Plugins/Lumin/MagicLeap/Source/MLSDK/Public/MagicLeapPluginUtil.h
--- a/Engine/Plugins/Lumin/MagicLeap/Source/MLSDK/Public/MagicLeapPluginUtil.h
+++ b/Engine/Plugins/Lumin/MagicLeap/Source/MLSDK/Public/MagicLeapPluginUtil.h
@@ -38,6 +38,15 @@
 #include "HAL/PlatformProcess.h"
 #include "Misc/Paths.h"
 
+#if PLATFORM_WINDOWS
+#define ML_INCLUDES_START THIRD_PARTY_INCLUDES_START \
+	__pragma(warning(disable: 4201)) /* warning C4201: nonstandard extension used: nameless struct/union */
+#else
+#define ML_INCLUDES_START THIRD_PARTY_INCLUDES_START
+#endif
+
+#define ML_INCLUDES_END THIRD_PARTY_INCLUDES_END
+
 DEFINE_LOG_CATEGORY_STATIC(LogMagicLeapAPISetup, Display, All);
 
 /** Utility class to load the correct MLSDK libs depedning on the path set to the MLSDK package and whether or not we want to use MLremote / Zero Iteration*/
diff --git a/Engine/Plugins/Lumin/MagicLeap/Source/MagicLeap/Private/AppFramework.cpp b/Engine/Plugins/Lumin/MagicLeap/Source/MagicLeap/Private/AppFramework.cpp
--- a/Engine/Plugins/Lumin/MagicLeap/Source/MagicLeap/Private/AppFramework.cpp
+++ b/Engine/Plugins/Lumin/MagicLeap/Source/MagicLeap/Private/AppFramework.cpp
@@ -38,9 +38,12 @@
 #include "RenderingThread.h"
 #include "Engine/Engine.h"
 #include "Misc/CoreDelegates.h"
+#include "MagicLeapPluginUtil.h" // for ML_INCLUDES_START/END
 
 #if WITH_MLSDK
+ML_INCLUDES_START
 #include <ml_snapshot.h>
+ML_INCLUDES_END
 #endif //WITH_MLSDK
 
 TArray<MagicLeap::IAppEventHandler*> FAppFramework::EventHandlers;
diff --git a/Engine/Plugins/Lumin/MagicLeap/Source/MagicLeap/Private/AppFramework.h b/Engine/Plugins/Lumin/MagicLeap/Source/MagicLeap/Private/AppFramework.h
--- a/Engine/Plugins/Lumin/MagicLeap/Source/MagicLeap/Private/AppFramework.h
+++ b/Engine/Plugins/Lumin/MagicLeap/Source/MagicLeap/Private/AppFramework.h
@@ -35,8 +35,11 @@
 #pragma once
 
 #include "GenericPlatform/GenericPlatformProcess.h"
+#include "MagicLeapPluginUtil.h" // for ML_INCLUDES_START/END
 
 #if WITH_MLSDK
+ML_INCLUDES_START
+
 #include <ml_api.h>
 
 #if PLATFORM_LUMIN
@@ -48,6 +51,7 @@
 
 #include <ml_coordinate_frame_uid.h>
 
+ML_INCLUDES_END
 #endif  //WITH_MLSDK
 
 #include "CoreMinimal.h"
diff --git a/Engine/Plugins/Lumin/MagicLeap/Source/MagicLeap/Private/CameraCaptureComponent.cpp b/Engine/Plugins/Lumin/MagicLeap/Source/MagicLeap/Private/CameraCaptureComponent.cpp
--- a/Engine/Plugins/Lumin/MagicLeap/Source/MagicLeap/Private/CameraCaptureComponent.cpp
+++ b/Engine/Plugins/Lumin/MagicLeap/Source/MagicLeap/Private/CameraCaptureComponent.cpp
@@ -19,11 +19,14 @@
 #if PLATFORM_LUMIN
 #include "Lumin/LuminAffinity.h"
 #endif // PLATFORM_LUMIN
+#include "MagicLeapPluginUtil.h" // for ML_INCLUDES_START/END
 
 #if WITH_MLSDK
+ML_INCLUDES_START
 #include <ml_camera.h>
 #include <ml_camera_metadata.h>
 #include <ml_camera_metadata_tags.h>
+ML_INCLUDES_END
 #endif //WITH_MLSDK
 
 DEFINE_LOG_CATEGORY(LogCameraCapture);
diff --git a/Engine/Plugins/Lumin/MagicLeap/Source/MagicLeap/Private/ImageTrackerComponent.cpp b/Engine/Plugins/Lumin/MagicLeap/Source/MagicLeap/Private/ImageTrackerComponent.cpp
--- a/Engine/Plugins/Lumin/MagicLeap/Source/MagicLeap/Private/ImageTrackerComponent.cpp
+++ b/Engine/Plugins/Lumin/MagicLeap/Source/MagicLeap/Private/ImageTrackerComponent.cpp
@@ -47,13 +47,16 @@
 #if WITH_EDITOR
 #include "Editor.h"
 #endif
+#include "MagicLeapPluginUtil.h" // for ML_INCLUDES_START/END
 
 #if WITH_MLSDK
+ML_INCLUDES_START
 #include <ml_image_tracking.h>
 #include <ml_perception.h>
 #include <ml_snapshot.h>
 #include <ml_head_tracking.h>
 #include <ml_coordinate_frame_uid.h>
+ML_INCLUDES_END
 #endif //WITH_MLSDK
 
 class FImageTrackerEngineInterface : public MagicLeap::IAppEventHandler
diff --git a/Engine/Plugins/Lumin/MagicLeap/Source/MagicLeap/Private/MagicLeapCustomPresent.cpp b/Engine/Plugins/Lumin/MagicLeap/Source/MagicLeap/Private/MagicLeapCustomPresent.cpp
--- a/Engine/Plugins/Lumin/MagicLeap/Source/MagicLeap/Private/MagicLeapCustomPresent.cpp
+++ b/Engine/Plugins/Lumin/MagicLeap/Source/MagicLeap/Private/MagicLeapCustomPresent.cpp
@@ -35,8 +35,12 @@
 #include "MagicLeapCustomPresent.h"
 #include "MagicLeapHMD.h"
 #include "RenderingThread.h"
+#include "MagicLeapPluginUtil.h" // for ML_INCLUDES_START/END
+
 #if WITH_MLSDK
+ML_INCLUDES_START
 #include <ml_lifecycle.h>
+ML_INCLUDES_END
 #endif //WITH_MLSDK
 
 #if PLATFORM_WINDOWS || PLATFORM_LINUX || PLATFORM_LUMIN
diff --git a/Engine/Plugins/Lumin/MagicLeap/Source/MagicLeap/Private/MagicLeapCustomPresent.h b/Engine/Plugins/Lumin/MagicLeap/Source/MagicLeap/Private/MagicLeapCustomPresent.h
--- a/Engine/Plugins/Lumin/MagicLeap/Source/MagicLeap/Private/MagicLeapCustomPresent.h
+++ b/Engine/Plugins/Lumin/MagicLeap/Source/MagicLeap/Private/MagicLeapCustomPresent.h
@@ -39,8 +39,10 @@
 
 #include "MagicLeapMath.h"
 #include "MagicLeapUtils.h"
+#include "MagicLeapPluginUtil.h" // for ML_INCLUDES_START/END
 
 #if WITH_MLSDK
+ML_INCLUDES_START
 #include <ml_api.h>
 
 #if PLATFORM_LUMIN
@@ -51,6 +53,7 @@
 #endif // PLATFORM_LUMIN
 
 #include <ml_snapshot.h>
+ML_INCLUDES_END
 #endif //WITH_MLSDK
 
 class FMagicLeapHMD;
diff --git a/Engine/Plugins/Lumin/MagicLeap/Source/MagicLeap/Private/MagicLeapHMD.cpp b/Engine/Plugins/Lumin/MagicLeap/Source/MagicLeap/Private/MagicLeapHMD.cpp
--- a/Engine/Plugins/Lumin/MagicLeap/Source/MagicLeap/Private/MagicLeapHMD.cpp
+++ b/Engine/Plugins/Lumin/MagicLeap/Source/MagicLeap/Private/MagicLeapHMD.cpp
@@ -68,7 +68,10 @@
 #include "OpenGLDrv.h"
 #endif
 
+#include "MagicLeapPluginUtil.h" // for ML_INCLUDES_START/END
+
 #if WITH_MLSDK
+ML_INCLUDES_START
 #include <ml_perception.h>
 
 #if !PLATFORM_LUMIN
@@ -87,6 +90,7 @@
 #include <ml_graphics.h>
 #endif // PLATFORM_LUMIN
 
+ML_INCLUDES_END
 #endif //WITH_MLSDK
 
 #if WITH_EDITOR
diff --git a/Engine/Plugins/Lumin/MagicLeap/Source/MagicLeap/Private/MagicLeapHMD.h b/Engine/Plugins/Lumin/MagicLeap/Source/MagicLeap/Private/MagicLeapHMD.h
--- a/Engine/Plugins/Lumin/MagicLeap/Source/MagicLeap/Private/MagicLeapHMD.h
+++ b/Engine/Plugins/Lumin/MagicLeap/Source/MagicLeap/Private/MagicLeapHMD.h
@@ -47,9 +47,12 @@
 #include "MagicLeapCustomPresent.h"
 #include "MagicLeapHMDFunctionLibrary.h"
 #include "LuminRuntimeSettings.h"
+#include "MagicLeapPluginUtil.h" // for ML_INCLUDES_START/END
 
 #if WITH_MLSDK
+ML_INCLUDES_START
 #include <ml_head_tracking.h>
+ML_INCLUDES_END
 #endif //WITH_MLSDK
 
 /**
diff --git a/Engine/Plugins/Lumin/MagicLeap/Source/MagicLeap/Private/MagicLeapUtils.h b/Engine/Plugins/Lumin/MagicLeap/Source/MagicLeap/Private/MagicLeapUtils.h
--- a/Engine/Plugins/Lumin/MagicLeap/Source/MagicLeap/Private/MagicLeapUtils.h
+++ b/Engine/Plugins/Lumin/MagicLeap/Source/MagicLeap/Private/MagicLeapUtils.h
@@ -34,8 +34,12 @@
 
 #pragma once
 
+#include "MagicLeapPluginUtil.h" // for ML_INCLUDES_START/END
+
 #if WITH_MLSDK
+ML_INCLUDES_START
 #include <ml_graphics.h>
+ML_INCLUDES_END
 
 namespace MagicLeap
 {
diff --git a/Engine/Plugins/Lumin/MagicLeap/Source/MagicLeap/Private/MeshTrackerComponent.cpp b/Engine/Plugins/Lumin/MagicLeap/Source/MagicLeap/Private/MeshTrackerComponent.cpp
--- a/Engine/Plugins/Lumin/MagicLeap/Source/MagicLeap/Private/MeshTrackerComponent.cpp
+++ b/Engine/Plugins/Lumin/MagicLeap/Source/MagicLeap/Private/MeshTrackerComponent.cpp
@@ -47,10 +47,13 @@
 #include "Editor.h"
 #endif
 #include "MRMeshComponent.h"
+#include "MagicLeapPluginUtil.h" // for ML_INCLUDES_START/END
 
 #if WITH_MLSDK
+ML_INCLUDES_START
 #include <ml_meshing.h>
 #include <ml_data_array.h>
+ML_INCLUDES_END
 #endif //WITH_MLSDK
 
 static TAutoConsoleVariable<float> CVarFakeMeshTrackerData(
diff --git a/Engine/Plugins/Lumin/MagicLeap/Source/MagicLeap/Private/RaycastComponent.cpp b/Engine/Plugins/Lumin/MagicLeap/Source/MagicLeap/Private/RaycastComponent.cpp
--- a/Engine/Plugins/Lumin/MagicLeap/Source/MagicLeap/Private/RaycastComponent.cpp
+++ b/Engine/Plugins/Lumin/MagicLeap/Source/MagicLeap/Private/RaycastComponent.cpp
@@ -41,9 +41,12 @@
 #if WITH_EDITOR
 #include "Editor.h"
 #endif
+#include "MagicLeapPluginUtil.h" // for ML_INCLUDES_START/END
 
 #if WITH_MLSDK
+ML_INCLUDES_START
 #include <ml_raycast.h>
+ML_INCLUDES_END
 #endif //WITH_MLSDK
 
 class FRaycastTrackerImpl
diff --git a/Engine/Plugins/Lumin/MagicLeap/Source/MagicLeapAudio/Private/AudioMixerPlatformMagicLeap.h b/Engine/Plugins/Lumin/MagicLeap/Source/MagicLeapAudio/Private/AudioMixerPlatformMagicLeap.h
--- a/Engine/Plugins/Lumin/MagicLeap/Source/MagicLeapAudio/Private/AudioMixerPlatformMagicLeap.h
+++ b/Engine/Plugins/Lumin/MagicLeap/Source/MagicLeapAudio/Private/AudioMixerPlatformMagicLeap.h
@@ -3,8 +3,12 @@
 #pragma once
 
 #include "AudioMixer.h"
+#include "MagicLeapPluginUtil.h" // for ML_INCLUDES_START/END
+
 #if WITH_MLSDK
+ML_INCLUDES_START
 #include <ml_audio.h>
+ML_INCLUDES_END
 #endif //WITH_MLSDK
 
 // Any platform defines
diff --git a/Engine/Plugins/Lumin/MagicLeap/Source/MagicLeapController/Private/MagicLeapController.h b/Engine/Plugins/Lumin/MagicLeap/Source/MagicLeapController/Private/MagicLeapController.h
--- a/Engine/Plugins/Lumin/MagicLeap/Source/MagicLeapController/Private/MagicLeapController.h
+++ b/Engine/Plugins/Lumin/MagicLeap/Source/MagicLeapController/Private/MagicLeapController.h
@@ -42,9 +42,12 @@
 #include "XRMotionControllerBase.h"
 #include "IMagicLeapControllerPlugin.h"
 #include "Misc/ScopeLock.h"
+#include "MagicLeapPluginUtil.h" // for ML_INCLUDES_START/END
 
 #if WITH_MLSDK
+ML_INCLUDES_START
 #include <ml_input.h>
+ML_INCLUDES_END
 #endif //WITH_MLSDK
 
 #include "CoreMinimal.h"
diff --git a/Engine/Plugins/Lumin/MagicLeap/Source/MagicLeapController/Private/Tests/MagicLeapControllerTests.cpp b/Engine/Plugins/Lumin/MagicLeap/Source/MagicLeapController/Private/Tests/MagicLeapControllerTests.cpp
--- a/Engine/Plugins/Lumin/MagicLeap/Source/MagicLeapController/Private/Tests/MagicLeapControllerTests.cpp
+++ b/Engine/Plugins/Lumin/MagicLeap/Source/MagicLeapController/Private/Tests/MagicLeapControllerTests.cpp
@@ -43,9 +43,12 @@
 #include "MagicLeapControllerFunctionLibrary.h"
 #include "Framework/Application/SlateApplication.h"
 #include "Haptics/HapticFeedbackEffect_Base.h"
+#include "MagicLeapPluginUtil.h" // for ML_INCLUDES_START/END
 
 #if WITH_MLSDK
+ML_INCLUDES_START
 #include <ml_input.h>
+ML_INCLUDES_END
 #endif //WITH_MLSDK
 
 DEFINE_LOG_CATEGORY_STATIC(LogMagicLeapControllerTest, Display, All);
diff --git a/Engine/Plugins/Lumin/MagicLeap/Source/MagicLeapEyeTracker/Private/MagicLeapVREyeTracker.cpp b/Engine/Plugins/Lumin/MagicLeap/Source/MagicLeapEyeTracker/Private/MagicLeapVREyeTracker.cpp
--- a/Engine/Plugins/Lumin/MagicLeap/Source/MagicLeapEyeTracker/Private/MagicLeapVREyeTracker.cpp
+++ b/Engine/Plugins/Lumin/MagicLeap/Source/MagicLeapEyeTracker/Private/MagicLeapVREyeTracker.cpp
@@ -47,9 +47,12 @@
 #include "UnrealEngine.h"
 #include "MagicLeapSDKDetection.h"
 #include "HeadMountedDisplayFunctionLibrary.h"
+#include "MagicLeapPluginUtil.h" // for ML_INCLUDES_START/END
 
 #if WITH_MLSDK
+ML_INCLUDES_START
 #include <ml_eye_tracking.h>
+ML_INCLUDES_END
 #endif //WITH_MLSDK
 
 FMagicLeapVREyeTracker::FMagicLeapVREyeTracker()
diff --git a/Engine/Plugins/Lumin/MagicLeap/Source/MagicLeapEyeTracker/Private/MagicLeapVREyeTracker.h b/Engine/Plugins/Lumin/MagicLeap/Source/MagicLeapEyeTracker/Private/MagicLeapVREyeTracker.h
--- a/Engine/Plugins/Lumin/MagicLeap/Source/MagicLeapEyeTracker/Private/MagicLeapVREyeTracker.h
+++ b/Engine/Plugins/Lumin/MagicLeap/Source/MagicLeapEyeTracker/Private/MagicLeapVREyeTracker.h
@@ -41,10 +41,13 @@
 #include "MagicLeapEyeTrackerTypes.h"
 #include "Containers/Ticker.h"
 #include "GameFramework/HUD.h"
+#include "MagicLeapPluginUtil.h" // for ML_INCLUDES_START/END
 
 #if WITH_MLSDK
+ML_INCLUDES_START
 #include <ml_api.h>
 #include <ml_eye_tracking.h>
+ML_INCLUDES_END
 #endif //WITH_MLSDK
 
 class FMagicLeapVREyeTracker : public IMagicLeapVREyeTracker, public FTickerObjectBase
diff --git a/Engine/Plugins/Lumin/MagicLeap/Source/MagicLeapGestures/Private/MagicLeapGestures.h b/Engine/Plugins/Lumin/MagicLeap/Source/MagicLeapGestures/Private/MagicLeapGestures.h
--- a/Engine/Plugins/Lumin/MagicLeap/Source/MagicLeapGestures/Private/MagicLeapGestures.h
+++ b/Engine/Plugins/Lumin/MagicLeap/Source/MagicLeapGestures/Private/MagicLeapGestures.h
@@ -36,9 +36,12 @@
 
 #include "GenericPlatform/IInputInterface.h"
 #include "IMagicLeapInputDevice.h"
+#include "MagicLeapPluginUtil.h" // for ML_INCLUDES_START/END
 
 #if WITH_MLSDK
+ML_INCLUDES_START
 #include <ml_gesture.h>
+ML_INCLUDES_END
 #endif //WITH_MLSDK
 
 #include "MagicLeapGestureTypes.h"
diff --git a/Engine/Plugins/Lumin/MagicLeapAnalytics/Source/MagicLeapAnalytics/Private/MagicLeapAnalytics.cpp b/Engine/Plugins/Lumin/MagicLeapAnalytics/Source/MagicLeapAnalytics/Private/MagicLeapAnalytics.cpp
--- a/Engine/Plugins/Lumin/MagicLeapAnalytics/Source/MagicLeapAnalytics/Private/MagicLeapAnalytics.cpp
+++ b/Engine/Plugins/Lumin/MagicLeapAnalytics/Source/MagicLeapAnalytics/Private/MagicLeapAnalytics.cpp
@@ -44,7 +44,9 @@
 #include "MagicLeapPluginUtil.h"
 
 #if WITH_MLSDK
+ML_INCLUDES_START
 #include <ml_app_analytics.h>
+ML_INCLUDES_END
 #endif //WITH_MLSDK
 
 DEFINE_LOG_CATEGORY_STATIC(LogMagicLeapAnalytics, Display, All);
diff --git a/Engine/Plugins/Lumin/MagicLeapScreens/Source/MagicLeapScreens/Private/ScreensComponent.cpp b/Engine/Plugins/Lumin/MagicLeapScreens/Source/MagicLeapScreens/Private/ScreensComponent.cpp
--- a/Engine/Plugins/Lumin/MagicLeapScreens/Source/MagicLeapScreens/Private/ScreensComponent.cpp
+++ b/Engine/Plugins/Lumin/MagicLeapScreens/Source/MagicLeapScreens/Private/ScreensComponent.cpp
@@ -17,10 +17,13 @@
 #if PLATFORM_LUMIN
 #include "Lumin/LuminAffinity.h"
 #endif // PLATFORM_LUMIN
+#include "MagicLeapPluginUtil.h" // for ML_INCLUDES_START/END
 
 #if WITH_MLSDK
+ML_INCLUDES_START
 #include <ml_image.h>
 #include <ml_screens.h>
+ML_INCLUDES_END
 #endif //WITH_MLSDK
 
 DEFINE_LOG_CATEGORY_STATIC(LogScreens, Display, All);
