Duplicating for 4.20:
Change: 4097595 and 4104555
UE-59359 Mr Mesh Crash at startup
-Vulkan does not work correctly unless we fill in all of the vertex buffers.  MRMesh now requires that via assert and MeshTrackerComponent fills in placeholder data as necessary.
-Also fixed the vr.MagicLeap.FakeMeshTrackerData cvar, it was missing the 'MagicLeap'.

diff --git a/Engine/Plugins/Lumin/MagicLeap/Source/MagicLeap/Private/MeshTrackerComponent.cpp b/Engine/Plugins/Lumin/MagicLeap/Source/MagicLeap/Private/MeshTrackerComponent.cpp
--- a/Engine/Plugins/Lumin/MagicLeap/Source/MagicLeap/Private/MeshTrackerComponent.cpp
+++ b/Engine/Plugins/Lumin/MagicLeap/Source/MagicLeap/Private/MeshTrackerComponent.cpp
@@ -487,7 +487,7 @@ void UMeshTrackerComponent::TickComponent(float DeltaTime, enum ELevelTick TickT
 		return;
 	}
 
-	static const auto FakeMeshTrackerDataCVar = IConsoleManager::Get().FindConsoleVariable(TEXT("vr.FakeMeshTrackerData"));
+	static const auto FakeMeshTrackerDataCVar = IConsoleManager::Get().FindConsoleVariable(TEXT("vr.MagicLeap.FakeMeshTrackerData"));
 	if (FakeMeshTrackerDataCVar != nullptr && FakeMeshTrackerDataCVar->GetInt())
 	{
 		TickWithFakeData();
@@ -674,27 +674,43 @@ void UMeshTrackerComponent::TickComponent(float DeltaTime, enum ELevelTick TickT
 					}
 
 					// Read normal stream
-					if (Impl->Data.normal_stream_index != MLDataArray_InvalidStreamIndex && CurrentMeshData.streams[Impl->Data.normal_stream_index].type != MLDataArrayType_None)
 					{
-						const uint32 NormalsCount = CurrentMeshData.streams[Impl->Data.normal_stream_index].count;
-						CurrentMeshDataCache->Normals.Reserve(NormalsCount);
-						for (uint32 i = 0; i < NormalsCount; ++i)
+						uint32 NormalsCount = 0;
+						if (Impl->Data.normal_stream_index != MLDataArray_InvalidStreamIndex && CurrentMeshData.streams[Impl->Data.normal_stream_index].type != MLDataArrayType_None)
+						{
+							NormalsCount = CurrentMeshData.streams[Impl->Data.normal_stream_index].count;
+							CurrentMeshDataCache->Normals.Reserve(NormalsCount);
+							for (uint32 i = 0; i < NormalsCount; ++i)
+							{
+								const auto &normal = CurrentMeshData.streams[Impl->Data.normal_stream_index].xyz_array[i];
+								CurrentMeshDataCache->Normals.Add(MagicLeap::ToFVector(normal, 1.0f));
+							}
+						}
+						else
 						{
-							const auto &normal = CurrentMeshData.streams[Impl->Data.normal_stream_index].xyz_array[i];
-							CurrentMeshDataCache->Normals.Add(MagicLeap::ToFVector(normal, 1.0f));
+							// Vulkan requires that tangent data exist, so we provide fake normals and tangents even if it does not
+
+							NormalsCount = CurrentMeshDataCache->Vertices.Num();
+							CurrentMeshDataCache->Normals.Reserve(NormalsCount);
+							for (uint32 i = 0; i < NormalsCount; ++i)
+							{
+								FVector Normal = CurrentMeshDataCache->Vertices[i];
+								Normal.Normalize();
+								CurrentMeshDataCache->Normals.Add(Normal);
+							}
 						}
 
 						// Also write normals into tangents
 						CurrentMeshDataCache->Tangents.Reserve(NormalsCount * 2);
 						for (uint32 i = 0; i < NormalsCount; ++i)
 						{
-							const auto &MLNormal = CurrentMeshData.streams[Impl->Data.normal_stream_index].xyz_array[i];
-							const FVector Normal = MagicLeap::ToFVector(MLNormal, 1.0f);
-							const FVector NonNormal = Normal.X < Normal.Z ? FVector(1,0,0) : FVector(0,1,0);
+							const FVector Normal = CurrentMeshDataCache->Normals[i];
+							const FVector NonNormal = Normal.X < Normal.Z ? FVector(0, 0, 1) : FVector(0, 1, 0);
 							const FVector TangentX = FVector::CrossProduct(Normal, NonNormal);
 							CurrentMeshDataCache->Tangents.Add(TangentX);
 							CurrentMeshDataCache->Tangents.Add(Normal);
 						}
+
 					}
 
 					// Read confidence stream
@@ -709,8 +725,7 @@ void UMeshTrackerComponent::TickComponent(float DeltaTime, enum ELevelTick TickT
 						}
 					}
 
-					// Write VertexColor, if necessary
-					if (VertexColorMode != EMLMeshVertexColorMode::None)
+					// Write VertexColor
 					{
 						switch (VertexColorMode)
 						{
@@ -747,6 +762,16 @@ void UMeshTrackerComponent::TickComponent(float DeltaTime, enum ELevelTick TickT
 							}
 							break;
 						}
+						case EMLMeshVertexColorMode::None:
+						{
+							// Vulkan requires that we fill everything in.
+							const uint32 VertexCount = CurrentMeshDataCache->Vertices.Num();
+							for (uint32 i = 0; i < VertexCount; ++i)
+							{
+								CurrentMeshDataCache->VertexColors.Add(FColor::White);
+							}
+							break;
+						}
 						default:
 							check(false);
 						}
@@ -756,6 +781,16 @@ void UMeshTrackerComponent::TickComponent(float DeltaTime, enum ELevelTick TickT
 					MLDataArrayUnlock(CurrentMeshHandle);
 				}
 
+				// Write UVs
+				{
+					const uint32 VertexCount = CurrentMeshDataCache->Vertices.Num();
+					for (uint32 i = 0; i < VertexCount; ++i)
+					{
+						const float FakeCoord = (float)i / (float)VertexCount;
+						CurrentMeshDataCache->UV0.Add(FVector2D(FakeCoord, FakeCoord));
+					}
+				}
+
 				// We don't add any sections to the procedural mesh component for point clouds.
 				if (MeshType != EMeshType::PointCloud)
 				{
@@ -904,9 +939,6 @@ void UMeshTrackerComponent::TickWithFakeData()
 
 				const int32 VertCount = 8;
 				CurrentMeshDataCache->Vertices.Reserve(VertCount);
-				CurrentMeshDataCache->UV0.Reserve(VertCount);
-				CurrentMeshDataCache->VertexColors.Reserve(VertCount);
-				CurrentMeshDataCache->Triangles.Reserve(6*2*3);
 
 				CurrentMeshDataCache->Vertices.Add(FVector(Origin.X + Extents.X, Origin.Y - Extents.Y, Origin.Z + Extents.Z));  // 0
 				CurrentMeshDataCache->Vertices.Add(FVector(Origin.X + Extents.X, Origin.Y + Extents.Y, Origin.Z + Extents.Z));  // 1
@@ -916,16 +948,36 @@ void UMeshTrackerComponent::TickWithFakeData()
 				CurrentMeshDataCache->Vertices.Add(FVector(Origin.X - Extents.X, Origin.Y + Extents.Y, Origin.Z + Extents.Z));  // 5
 				CurrentMeshDataCache->Vertices.Add(FVector(Origin.X - Extents.X, Origin.Y + Extents.Y, Origin.Z - Extents.Z));  // 6
 				CurrentMeshDataCache->Vertices.Add(FVector(Origin.X - Extents.X, Origin.Y - Extents.Y, Origin.Z - Extents.Z));  // 7
+
+				CurrentMeshDataCache->UV0.Reserve(VertCount);
+				CurrentMeshDataCache->VertexColors.Reserve(VertCount);
 				for (int i = 0; i < VertCount; ++i)
 				{
 					int imod = i % 10;
 					CurrentMeshDataCache->VertexColors.Emplace(FColor::MakeRandomColor());
 					CurrentMeshDataCache->UV0.Add(FVector2D(imod*0.1f, imod*0.1f));
-					//FVector Tangent(imod*0.1f, imod*imod*0.5f, 1.0f);
-					//Tangent.Normalize();
-					//CurrentMeshDataCache->Tangents.Add(Tangent);
 				}
+				
+				// Vulkan requires that tangent data exist
 
+				CurrentMeshDataCache->Normals.Reserve(VertCount);
+				for (uint32 i = 0; i < VertCount; ++i)
+				{
+					CurrentMeshDataCache->Normals.Add(FVector(0,0,1));
+				}
+
+				// Also write normals into tangents
+				CurrentMeshDataCache->Tangents.Reserve(VertCount * 2);
+				for (uint32 i = 0; i < VertCount; ++i)
+				{
+					const FVector Normal = CurrentMeshDataCache->Normals[i];
+					const FVector NonNormal = Normal.X < Normal.Z ? FVector(0, 0, 1) : FVector(0, 1, 0);
+					const FVector TangentX = FVector::CrossProduct(Normal, NonNormal);
+					CurrentMeshDataCache->Tangents.Add(TangentX);
+					CurrentMeshDataCache->Tangents.Add(Normal);
+				}
+				
+				CurrentMeshDataCache->Triangles.Reserve(6 * 2 * 3);
 				// Read triangle (indices) stream
 				// This makes half the triangles of a box, one per side.
 				//CurrentMeshDataCache->Triangles.Add(0);
diff --git a/Engine/Source/Runtime/MRMesh/Private/MRMeshComponent.cpp b/Engine/Source/Runtime/MRMesh/Private/MRMeshComponent.cpp
--- a/Engine/Source/Runtime/MRMesh/Private/MRMeshComponent.cpp
+++ b/Engine/Source/Runtime/MRMesh/Private/MRMeshComponent.cpp
@@ -138,12 +138,13 @@ static void InitVertexFactory(FLocalVertexFactory* VertexFactory, const FMRMeshP
 			NewData.PositionComponentSRV = MRMeshSection.PositionBufferSRV;
 			NewData.PositionComponent = FVertexStreamComponent(&MRMeshSection.PositionBuffer, 0, sizeof(FVector), VET_Float3, EVertexStreamUsage::Default);
 		}
-		
-		//{
-		//	NewData.TextureCoordinatesSRV = MRMeshSection.UVBufferSRV;
-		//	NewData.TextureCoordinates.Add(FVertexStreamComponent(&MRMeshSection.UVBuffer, 0, sizeof(FVector2D), VET_Float2, EVertexStreamUsage::ManualFetch));
-		//	NewData.NumTexCoords = 1;
-		//}
+
+		if (MRMeshSection.UVBuffer.NumVerts != 0)
+		{
+			NewData.TextureCoordinatesSRV = MRMeshSection.UVBufferSRV;
+			NewData.TextureCoordinates.Add(FVertexStreamComponent(&MRMeshSection.UVBuffer, 0, sizeof(FVector2D), VET_Float2, EVertexStreamUsage::ManualFetch));
+			NewData.NumTexCoords = 1;
+		}
 
 		if (MRMeshSection.TangentXZBuffer.NumVerts != 0)
 		{
@@ -199,9 +200,11 @@ public:
 		FMRMeshProxySection* NewSection = new FMRMeshProxySection(Args.BrickId, FeatureLevel);
 		ProxySections.Add(NewSection);
 
-		check((Args.PositionData.Num() == Args.ColorData.Num()) || Args.ColorData.Num() == 0);
-		check((Args.PositionData.Num() == Args.UVData.Num()) || Args.UVData.Num() == 0);
-		check(((Args.PositionData.Num() * 2) == Args.TangentXZData.Num()) || Args.TangentXZData.Num() == 0);
+		// Vulkan requires that all the buffers be full.
+		const int32 NumVerts = Args.PositionData.Num();
+		check((NumVerts == Args.ColorData.Num()));
+		check((NumVerts == Args.UVData.Num()));
+		check((NumVerts * 2) == Args.TangentXZData.Num());
 
 		// POSITION BUFFER
 		{
@@ -226,10 +229,11 @@ public:
 			if (Args.TangentXZData.Num())
 			{
 				NewSection->TangentXZBuffer.InitRHIWith(Args.TangentXZData);
-				if (RHISupportsManualVertexFetch(GMaxRHIShaderPlatform))
-				{
-					NewSection->TangentXZBufferSRV = RHICreateShaderResourceView(NewSection->TangentXZBuffer.VertexBufferRHI, 4, PF_R8G8B8A8_SNORM);
-				}
+			}
+
+			if (RHISupportsManualVertexFetch(GMaxRHIShaderPlatform))
+			{
+				NewSection->TangentXZBufferSRV = RHICreateShaderResourceView(NewSection->TangentXZBuffer.VertexBufferRHI, 4, PF_R8G8B8A8_SNORM);
 			}
 		}
 
