Copying CL 4122930 from Dev-VR

Use XR ThreadUtils when scheduling GameTrackingThread to Render and RHI thread.

diff --git a/Engine/Plugins/Lumin/MagicLeap/Source/MagicLeap/Private/MagicLeapHMD.cpp b/Engine/Plugins/Lumin/MagicLeap/Source/MagicLeap/Private/MagicLeapHMD.cpp
--- a/Engine/Plugins/Lumin/MagicLeap/Source/MagicLeap/Private/MagicLeapHMD.cpp
+++ b/Engine/Plugins/Lumin/MagicLeap/Source/MagicLeap/Private/MagicLeapHMD.cpp
@@ -1501,36 +1501,29 @@ void FMagicLeapHMD::GetEyeRenderParams_RenderThread(const FRenderingCompositePas
 	}
 }
 
-void BeginRenderingHelper_RenderThread(FMagicLeapHMD* HMD, const FTrackingFrame& TrackingFrameCopy)
-{
-#if WITH_MLSDK
-	MLSnapshot* OldSnapshot = HMD->RenderTrackingFrame.Snapshot;
-	HMD->RenderTrackingFrame = TrackingFrameCopy;
-#if !PLATFORM_MAC
-	ExecuteOnRHIThread_DoNotWait([HMD, TrackingFrameCopy]()
-	{
-		HMD->RHITrackingFrame = TrackingFrameCopy;
-	});
-#endif //PLATFORM_MAC
-#endif //WITH_MLSDK
-}
-
 void FMagicLeapHMD::OnBeginRendering_GameThread()
 {
 	check(IsInGameThread());
 
 	RefreshTrackingFrame();
 
+#if WITH_MLSDK
 	// Copy the game tracking frame to the render frame.
-	// Since we don't flush the render commands here, we copy the game frame twice:
+	// Since we don't flush the render commands here, we copy the game frame thrice:
 	// 1st copy when enqueuing the command
 	// 2nd copy on the render thread during the command execution
-	ENQUEUE_UNIQUE_RENDER_COMMAND_TWOPARAMETER(InitializeRenderFrameFromGameFrameCopy,
-		FMagicLeapHMD*, HMD, this,
-		FTrackingFrame, TrackingFrameCopy, GameTrackingFrame,
+	ExecuteOnRenderThread_DoNotWait([this, TrackingFrameCopy = GameTrackingFrame]()
+	{
+		MLSnapshot* OldSnapshot = RenderTrackingFrame.Snapshot;
+		RenderTrackingFrame = TrackingFrameCopy;
+#if !PLATFORM_MAC
+		ExecuteOnRHIThread_DoNotWait([this, TrackingFrameCopy]()
 		{
-			BeginRenderingHelper_RenderThread(HMD, TrackingFrameCopy);
+			RHITrackingFrame = TrackingFrameCopy;
 		});
+#endif //PLATFORM_MAC
+	});
+#endif //WITH_MLSDK
 }
 
 TSharedPtr<class IXRCamera, ESPMode::ThreadSafe> FMagicLeapHMD::GetXRCamera(int32 DeviceId /*= HMDDeviceId*/)
