Only attempting to load dlls if the MLSDK is set up (avoiding needless warnings for users who aren't using the MLSDK).

diff --git a/Engine/Plugins/Lumin/MagicLeap/Source/MLSDK/Public/MagicLeapPluginUtil.h b/Engine/Plugins/Lumin/MagicLeap/Source/MLSDK/Public/MagicLeapPluginUtil.h
--- a/Engine/Plugins/Lumin/MagicLeap/Source/MLSDK/Public/MagicLeapPluginUtil.h
+++ b/Engine/Plugins/Lumin/MagicLeap/Source/MLSDK/Public/MagicLeapPluginUtil.h
@@ -69,22 +69,34 @@ public:
 			{
 				DllSearchPaths.Add(VDZILibraryPath);
 			}
-			// The default VDZI dir.
-			DllSearchPaths.Add(FPaths::Combine(*MLSDK, TEXT("VirtualDevice"), TEXT("lib")));
+
+			if (!MLSDK.IsEmpty())
+			{
+				// The default VDZI dir.
+				DllSearchPaths.Add(FPaths::Combine(*MLSDK, TEXT("VirtualDevice"), TEXT("lib")));
+			}
 		}
 #endif
 
 		// The MLSDK DLLs are platform specific and are segregated in directories for each platform.
+
+#if PLATFORM_LUMIN
+		// Lumin uses the system path as we are in device.
+		DllSearchPaths.Add(TEXT("/system/lib64"));
+#else
+
+		if (!MLSDK.IsEmpty())
+		{
 #if PLATFORM_WINDOWS
-		DllSearchPaths.Add(FPaths::Combine(*MLSDK, TEXT("lib"), TEXT("win64")));
+			DllSearchPaths.Add(FPaths::Combine(*MLSDK, TEXT("lib"), TEXT("win64")));
 #elif PLATFORM_LINUX
-		DllSearchPaths.Add(FPaths::Combine(*MLSDK, TEXT("lib"), TEXT("linux64")));
+			DllSearchPaths.Add(FPaths::Combine(*MLSDK, TEXT("lib"), TEXT("linux64")));
 #elif PLATFORM_MAC
-		DllSearchPaths.Add(FPaths::Combine(*MLSDK, TEXT("lib"), TEXT("osx")));
-#elif PLATFORM_LUMIN
-		// Lumin uses the system path as we are in device.
-		DllSearchPaths.Add(TEXT("/system/lib64"));
+			DllSearchPaths.Add(FPaths::Combine(*MLSDK, TEXT("lib"), TEXT("osx")));
 #endif // PLATFORM_WINDOWS
+		}
+
+#endif // PLATFORM_LUMIN
 
 		// Add the search paths to where we will load the DLLs from. For all just add to the UE4
 		// directory listing. But for Windows we also need to manipulate the PATH for the load
