Setting bGenerateOverlapEvents back to private

diff --git a/Engine/Plugins/Lumin/MagicLeap/Source/MagicLeap/Private/FoundObjectComponent.cpp b/Engine/Plugins/Lumin/MagicLeap/Source/MagicLeap/Private/FoundObjectComponent.cpp
--- a/Engine/Plugins/Lumin/MagicLeap/Source/MagicLeap/Private/FoundObjectComponent.cpp
+++ b/Engine/Plugins/Lumin/MagicLeap/Source/MagicLeap/Private/FoundObjectComponent.cpp
@@ -141,7 +141,7 @@ UFoundObjectComponent::UFoundObjectComponent()
 	SearchVolume->SetCanEverAffectNavigation(false);
 	SearchVolume->CanCharacterStepUpOn = ECanBeCharacterBase::ECB_No;
 	SearchVolume->SetCollisionObjectType(ECollisionChannel::ECC_WorldDynamic);
-	SearchVolume->bGenerateOverlapEvents = false;
+	SearchVolume->SetGenerateOverlapEvents(false);
 
 #if WITH_EDITOR
 	if (GIsEditor)
diff --git a/Engine/Plugins/Lumin/MagicLeap/Source/MagicLeap/Private/MeshTrackerComponent.cpp b/Engine/Plugins/Lumin/MagicLeap/Source/MagicLeap/Private/MeshTrackerComponent.cpp
--- a/Engine/Plugins/Lumin/MagicLeap/Source/MagicLeap/Private/MeshTrackerComponent.cpp
+++ b/Engine/Plugins/Lumin/MagicLeap/Source/MagicLeap/Private/MeshTrackerComponent.cpp
@@ -367,7 +367,7 @@ UMeshTrackerComponent::UMeshTrackerComponent(const FObjectInitializer& ObjectIni
 	BoundingVolume->SetCanEverAffectNavigation(false);
 	BoundingVolume->CanCharacterStepUpOn = ECanBeCharacterBase::ECB_No;
 	BoundingVolume->SetCollisionObjectType(ECollisionChannel::ECC_WorldDynamic);
-	BoundingVolume->bGenerateOverlapEvents = false;
+	BoundingVolume->SetGenerateOverlapEvents(false);
 	// Recommended default box extents for meshing - 10m (5m radius)
 	BoundingVolume->SetBoxExtent(FVector(1000, 1000, 1000), false);
 
diff --git a/Engine/Plugins/Lumin/MagicLeap/Source/MagicLeap/Private/PlanesComponent.cpp b/Engine/Plugins/Lumin/MagicLeap/Source/MagicLeap/Private/PlanesComponent.cpp
--- a/Engine/Plugins/Lumin/MagicLeap/Source/MagicLeap/Private/PlanesComponent.cpp
+++ b/Engine/Plugins/Lumin/MagicLeap/Source/MagicLeap/Private/PlanesComponent.cpp
@@ -176,7 +176,7 @@ UPlanesComponent::UPlanesComponent()
 	SearchVolume->SetCanEverAffectNavigation(false);
 	SearchVolume->CanCharacterStepUpOn = ECanBeCharacterBase::ECB_No;
 	SearchVolume->SetCollisionObjectType(ECollisionChannel::ECC_WorldDynamic);
-	SearchVolume->bGenerateOverlapEvents = false;
+	SearchVolume->SetGenerateOverlapEvents(false);
 	// Recommended default box extents for meshing - 10m (5m radius)
 	SearchVolume->SetBoxExtent(FVector(1000, 1000, 1000), false);
 
diff --git a/Engine/Source/Runtime/Engine/Classes/Components/PrimitiveComponent.h b/Engine/Source/Runtime/Engine/Classes/Components/PrimitiveComponent.h
--- a/Engine/Source/Runtime/Engine/Classes/Components/PrimitiveComponent.h
+++ b/Engine/Source/Runtime/Engine/Classes/Components/PrimitiveComponent.h
@@ -269,8 +269,7 @@ public:
 	UFUNCTION(BlueprintSetter)
 	void SetGenerateOverlapEvents(bool bInGenerateOverlapEvents);
 
-//LUMIN_MERGE - Removed private to allow FoundObjectComponent access
-//private:
+private:
 	UPROPERTY(EditAnywhere, BlueprintGetter = GetGenerateOverlapEvents, BlueprintSetter = SetGenerateOverlapEvents, Category = Collision)
 	uint8 bGenerateOverlapEvents : 1;
 
