#DEV-VR - Adding verification that lib directory of MLSDK is present on for the builders

diff --git a/Engine/Plugins/Lumin/MagicLeap/Source/MLSDK/MLSDK.Build.cs b/Engine/Plugins/Lumin/MagicLeap/Source/MLSDK/MLSDK.Build.cs
--- a/Engine/Plugins/Lumin/MagicLeap/Source/MLSDK/MLSDK.Build.cs
+++ b/Engine/Plugins/Lumin/MagicLeap/Source/MLSDK/MLSDK.Build.cs
@@ -55,8 +55,6 @@ public class MLSDK : ModuleRules
 			string LibraryPath = Path.Combine(MLSDKPath, "lib");
 			string VirtualDeviceLibraryPath = Path.Combine(MLSDKPath, "VirtualDevice", "lib");
 
-			bIsMLSDKInstalled = Directory.Exists(IncludePath);
-
 			string ProjectFileName = null != Target.ProjectFile ? Target.ProjectFile.FullName : "";
 			DirectoryReference ProjectDir =
 				string.IsNullOrEmpty(ProjectFileName) ? (DirectoryReference)null : Target.ProjectFile.Directory;
@@ -95,45 +93,49 @@ public class MLSDK : ModuleRules
 			}
 
 			LibraryPath = Path.Combine(LibraryPath, LibraryPlatformFolder);
-			PublicLibraryPaths.Add(LibraryPath);
-			if (!string.IsNullOrEmpty(MLSDKLibraryPath))
-			{
-				PublicLibraryPaths.Add(MLSDKLibraryPath);
-			}
-
-			if (Target.Platform != UnrealTargetPlatform.Lumin)
-			{
-				PublicLibraryPaths.Add(VirtualDeviceLibraryPath);
-			}
 
-			string[] MLSDKLibraryList = new string[] {
-				"ml_audio",
-				"ml_camera_metadata",
-				"ml_camera",
-				"ml_dispatch",
-				"ml_ext_logging",
-				"ml_graphics",
-				"ml_identity",
-				"ml_input",
-				"ml_lifecycle",
-				"ml_mediacodeclist",
-				"ml_mediacodec",
-				"ml_mediacrypto",
-				"ml_mediadrm",
-				"ml_mediaerror",
-				"ml_mediaextractor",
-				"ml_mediaformat",
-				"ml_mediaplayer",
-				"ml_musicservice",
-				"ml_perception_client",
-				"ml_privileges",
-				"ml_screens",
-				"ml_secure_storage",
-				"ml_sharedfile"
-			};
+			//make sure the platform specific libs are there as well
+			bIsMLSDKInstalled = Directory.Exists(IncludePath) && Directory.Exists(LibraryPath);
 
 			if (bIsMLSDKInstalled)
 			{
+				PublicLibraryPaths.Add(LibraryPath);
+				if (!string.IsNullOrEmpty(MLSDKLibraryPath))
+				{
+					PublicLibraryPaths.Add(MLSDKLibraryPath);
+				}
+
+				if (Target.Platform != UnrealTargetPlatform.Lumin)
+				{
+					PublicLibraryPaths.Add(VirtualDeviceLibraryPath);
+				}
+
+				string[] MLSDKLibraryList = new string[] {
+					"ml_audio",
+					"ml_camera_metadata",
+					"ml_camera",
+					"ml_dispatch",
+					"ml_ext_logging",
+					"ml_graphics",
+					"ml_identity",
+					"ml_input",
+					"ml_lifecycle",
+					"ml_mediacodeclist",
+					"ml_mediacodec",
+					"ml_mediacrypto",
+					"ml_mediadrm",
+					"ml_mediaerror",
+					"ml_mediaextractor",
+					"ml_mediaformat",
+					"ml_mediaplayer",
+					"ml_musicservice",
+					"ml_perception_client",
+					"ml_privileges",
+					"ml_screens",
+					"ml_secure_storage",
+					"ml_sharedfile"
+				};
+
 				if (Target.Platform == UnrealTargetPlatform.Win64)
 				{
 					foreach (string libname in MLSDKLibraryList)
