#DEV-VR - Fix for static analysis always being true

diff --git a/Engine/Plugins/Lumin/MagicLeap/Source/MagicLeapAudio/Private/AudioMixerPlatformMagicLeap.cpp b/Engine/Plugins/Lumin/MagicLeap/Source/MagicLeapAudio/Private/AudioMixerPlatformMagicLeap.cpp
--- a/Engine/Plugins/Lumin/MagicLeap/Source/MagicLeapAudio/Private/AudioMixerPlatformMagicLeap.cpp
+++ b/Engine/Plugins/Lumin/MagicLeap/Source/MagicLeapAudio/Private/AudioMixerPlatformMagicLeap.cpp
@@ -322,7 +322,8 @@ namespace Audio
 			return false;
 		}
 
-		if (AudioStreamInfo.StreamState != EAudioOutputStreamState::Stopped)
+		//we have already returned if we're not running.  Therefore, if here, we are both initialized AND running.
+		//if (AudioStreamInfo.StreamState != EAudioOutputStreamState::Stopped)
 		{
 			if (!MLAudioStopSound(StreamHandle, &Result))
 			{
