Renaming bPinActor to bShouldPinActor to avoid same name when exposed to Python

diff --git a/Engine/Plugins/Lumin/MagicLeap/Source/MagicLeap/Classes/MagicLeapARPinComponent.h b/Engine/Plugins/Lumin/MagicLeap/Source/MagicLeap/Classes/MagicLeapARPinComponent.h
--- a/Engine/Plugins/Lumin/MagicLeap/Source/MagicLeap/Classes/MagicLeapARPinComponent.h
+++ b/Engine/Plugins/Lumin/MagicLeap/Source/MagicLeap/Classes/MagicLeapARPinComponent.h
@@ -204,7 +204,7 @@ public:
 
 	/** Pin this component's owner actor instead of just the component itself. Relevant only when using 'OnlyOnDataRestoration' or 'Always' as AutoPinType. */
 	UPROPERTY(BlueprintReadWrite, EditAnywhere, Category = "ContentPersistence|MagicLeap")
-	bool bPinActor;
+	bool bShouldPinActor;
 
 	/**
 	 * Delegate used to notify the instigating blueprint that an entity (component or actor) has been successfuly pinned to the real-world.
diff --git a/Engine/Plugins/Lumin/MagicLeap/Source/MagicLeap/Private/MagicLeapARPinComponent.cpp b/Engine/Plugins/Lumin/MagicLeap/Source/MagicLeap/Private/MagicLeapARPinComponent.cpp
--- a/Engine/Plugins/Lumin/MagicLeap/Source/MagicLeap/Private/MagicLeapARPinComponent.cpp
+++ b/Engine/Plugins/Lumin/MagicLeap/Source/MagicLeap/Private/MagicLeapARPinComponent.cpp
@@ -375,7 +375,7 @@ bool UMagicLeapARPinFunctionLibrary::GetARPinPositionAndOrientation(const FGuid&
 
 UMagicLeapARPinComponent::UMagicLeapARPinComponent()
 : AutoPinType(EAutoPinType::OnlyOnDataRestoration)
-, bPinActor(false)
+, bShouldPinActor(false)
 , PinnedSceneComponent(nullptr)
 , OldComponentWorldTransform(FTransform::Identity)
 , OldCFUIDTransform(FTransform::Identity)
@@ -425,7 +425,7 @@ void UMagicLeapARPinComponent::BeginPlay()
 
 	if ((AutoPinType == EAutoPinType::Always) || (bDataRestored && AutoPinType == EAutoPinType::OnlyOnDataRestoration))
 	{
-		if (bPinActor)
+		if (bShouldPinActor)
 		{
 			PinActor(GetOwner());
 		}
