Cleanup blueprint callable categories

x|Magic Leap -> Magic Leap|x
Apple ARKit * -> Apple ARKit|*

diff --git a/Engine/Plugins/Lumin/MagicLeap/Source/MagicLeap/Classes/CameraCaptureComponent.h b/Engine/Plugins/Lumin/MagicLeap/Source/MagicLeap/Classes/CameraCaptureComponent.h
--- a/Engine/Plugins/Lumin/MagicLeap/Source/MagicLeap/Classes/CameraCaptureComponent.h
+++ b/Engine/Plugins/Lumin/MagicLeap/Source/MagicLeap/Classes/CameraCaptureComponent.h
@@ -36,7 +36,7 @@ public:
 		   Upon completion, a successful operation will provide the file path of the newly created jpeg
 		   to the FCameraCaptureImgToFileSuccess event handler.
   */
-  UFUNCTION(BlueprintCallable, Category = "CameraCapture|MagicLeap")
+  UFUNCTION(BlueprintCallable, Category = "MagicLeap|CameraCapture")
   bool CaptureImageToFileAsync();
 
   /** 
@@ -46,7 +46,7 @@ public:
 		   image to the FCameraCaptureImgToTextureSuccess event handler.
     @note The generated texture will be garbage collected when this class is destroyed.
   */
-  UFUNCTION(BlueprintCallable, Category = "CameraCapture|MagicLeap")
+  UFUNCTION(BlueprintCallable, Category = "MagicLeap|CameraCapture")
   bool CaptureImageToTextureAsync();
 
   /** 
@@ -57,7 +57,7 @@ public:
 		   FCameraCaptureVidToFileSuccess event handler.
 	@param VideoLength The length in seconds of the footage to be captured.
   */
-  UFUNCTION(BlueprintCallable, Category = "CameraCapture|MagicLeap")
+  UFUNCTION(BlueprintCallable, Category = "MagicLeap|CameraCapture")
   bool CaptureVideoToFileAsync(float VideoLength);
 
   /**
diff --git a/Engine/Plugins/Lumin/MagicLeap/Source/MagicLeap/Classes/FoundObjectComponent.h b/Engine/Plugins/Lumin/MagicLeap/Source/MagicLeap/Classes/FoundObjectComponent.h
--- a/Engine/Plugins/Lumin/MagicLeap/Source/MagicLeap/Classes/FoundObjectComponent.h
+++ b/Engine/Plugins/Lumin/MagicLeap/Source/MagicLeap/Classes/FoundObjectComponent.h
@@ -119,7 +119,7 @@ public:
 	  @param ResultDelegate Delegate which will be called when the found object result is ready.
 	  @returns True if the found object query was successfully placed, false otherwise.
 	*/
-	UFUNCTION(BlueprintCallable, Category = "FoundObject|MagicLeap")
+	UFUNCTION(BlueprintCallable, Category = "MagicLeap|FoundObject")
 	bool SubmitQuery(int32& QueryID, const FFoundObjectResultDelegate& ResultDelegate);
 
 	virtual void TickComponent(float DeltaTime, enum ELevelTick TickType, FActorComponentTickFunction *ThisTickFunction) override;
diff --git a/Engine/Plugins/Lumin/MagicLeap/Source/MagicLeap/Classes/MeshTrackerComponent.h b/Engine/Plugins/Lumin/MagicLeap/Source/MagicLeap/Classes/MeshTrackerComponent.h
--- a/Engine/Plugins/Lumin/MagicLeap/Source/MagicLeap/Classes/MeshTrackerComponent.h
+++ b/Engine/Plugins/Lumin/MagicLeap/Source/MagicLeap/Classes/MeshTrackerComponent.h
@@ -176,7 +176,7 @@ public:
 	*/
 
 	/** Force an update on a non-live updating mesh tracker. */
-	UFUNCTION(BlueprintCallable, Category = "Meshing|MagicLeap")
+	UFUNCTION(BlueprintCallable, Category = "MagicLeap|Meshing")
 	bool ForceMeshUpdate();
 
 	/** Polls for and handles the results of the environmental mesh queries. */
diff --git a/Engine/Plugins/Lumin/MagicLeap/Source/MagicLeap/Classes/PlanesComponent.h b/Engine/Plugins/Lumin/MagicLeap/Source/MagicLeap/Classes/PlanesComponent.h
--- a/Engine/Plugins/Lumin/MagicLeap/Source/MagicLeap/Classes/PlanesComponent.h
+++ b/Engine/Plugins/Lumin/MagicLeap/Source/MagicLeap/Classes/PlanesComponent.h
@@ -125,7 +125,7 @@ public:
 	  @param ResultDelegate Delegate which will be called when the planes result is ready.
 	  @returns True if the planes query was successfully placed, false otherwise.
 	*/
-	UFUNCTION(BlueprintCallable, Category = "Planes|MagicLeap")
+	UFUNCTION(BlueprintCallable, Category = "MagicLeap|Planes")
 	bool RequestPlanes(int32 UserData, const FPlaneResultDelegate& ResultDelegate);
 
 	/** Polls for and handles the results of the plane queries. */
diff --git a/Engine/Plugins/Runtime/AppleARKit/Source/AppleARKit/Public/AppleARKitAnchor.h b/Engine/Plugins/Runtime/AppleARKit/Source/AppleARKit/Public/AppleARKitAnchor.h
--- a/Engine/Plugins/Runtime/AppleARKit/Source/AppleARKit/Public/AppleARKitAnchor.h
+++ b/Engine/Plugins/Runtime/AppleARKit/Source/AppleARKit/Public/AppleARKitAnchor.h
@@ -38,7 +38,7 @@ public:
 	 *
 	 * @todo arkit Fix this ^
 	 */
-	UFUNCTION( BlueprintPure, Category = "AppleARKitAnchor" )
+	UFUNCTION( BlueprintPure, Category = "AppleARKit|Anchor" )
 	FTransform GetTransform() const;
 
 #if SUPPORTS_ARKIT_1_0
diff --git a/Engine/Plugins/Runtime/AppleARKit/Source/AppleARKit/Public/AppleARKitPlaneAnchor.h b/Engine/Plugins/Runtime/AppleARKit/Source/AppleARKit/Public/AppleARKitPlaneAnchor.h
--- a/Engine/Plugins/Runtime/AppleARKit/Source/AppleARKit/Public/AppleARKitPlaneAnchor.h
+++ b/Engine/Plugins/Runtime/AppleARKit/Source/AppleARKit/Public/AppleARKitPlaneAnchor.h
@@ -24,16 +24,16 @@ public:
 	/**
 	 * The center of the plane in the anchor’s coordinate space.
 	 */
-	UFUNCTION( BlueprintPure, Category = "AppleARKitPlaneAnchor" )
+	UFUNCTION( BlueprintPure, Category = "AppleARKit|PlaneAnchor" )
 	FVector GetCenter() const;
 
 	/**
 	 * The extent of the plane in the anchor’s coordinate space.
 	 */
-	UFUNCTION( BlueprintPure, Category = "AppleARKitPlaneAnchor")
+	UFUNCTION( BlueprintPure, Category = "AppleARKit|PlaneAnchor")
 	FVector GetExtent() const;
 
-	UFUNCTION( BlueprintPure, Category = "AppleARKitPlaneAnchor")
+	UFUNCTION( BlueprintPure, Category = "AppleARKit|PlaneAnchor")
 	FTransform GetTransformToCenter() const;
 
 #if SUPPORTS_ARKIT_1_0
