UE-70084 Magic Leap blueprint Functions/Enums fail to load properly when reopening a project
-Changed MagicLeapHandTracking module LoadingPhase to PreDefault, almost all other blueprint libraries seem to have earlier-than-default loadingphase.

diff --git a/Engine/Plugins/Lumin/MagicLeap/MagicLeap.uplugin b/Engine/Plugins/Lumin/MagicLeap/MagicLeap.uplugin
--- a/Engine/Plugins/Lumin/MagicLeap/MagicLeap.uplugin
+++ b/Engine/Plugins/Lumin/MagicLeap/MagicLeap.uplugin
@@ -72,6 +72,7 @@
 		{
 			"Name" : "MagicLeapHandTracking",
 			"Type" : "Runtime",
+			"LoadingPhase" : "PreDefault",
 			"WhitelistPlatforms" :
 			[
 				"Win64",
