Re-enabling ML by default now that the ZI crash has been resolved.

diff --git a/Engine/Plugins/Lumin/MagicLeap/MagicLeap.uplugin b/Engine/Plugins/Lumin/MagicLeap/MagicLeap.uplugin
--- a/Engine/Plugins/Lumin/MagicLeap/MagicLeap.uplugin
+++ b/Engine/Plugins/Lumin/MagicLeap/MagicLeap.uplugin
@@ -7,7 +7,7 @@
 	"CreatedByURL" : "http://magicleap.com",
 	"Description" : "Support for the Magic Leap Head Mounted Display.",
 	"Category" : "Augmented Reality",
-	"EnabledByDefault" : false,
+	"EnabledByDefault" : true,
 	"DocsURL" : "",
 	"MarketplaceURL" : "",
 	"SupportURL" : "",
