[NO NON-PII COMMIT MESSAGE]

diff --git a/Engine/Plugins/Lumin/MagicLeap/Source/MagicLeapController/Private/MagicLeapController.cpp b/Engine/Plugins/Lumin/MagicLeap/Source/MagicLeapController/Private/MagicLeapController.cpp
--- a/Engine/Plugins/Lumin/MagicLeap/Source/MagicLeapController/Private/MagicLeapController.cpp
+++ b/Engine/Plugins/Lumin/MagicLeap/Source/MagicLeapController/Private/MagicLeapController.cpp
@@ -792,7 +792,7 @@ void FMagicLeapController::UpdateControllerStateFromInputTracker(const FAppFrame
 
 	// Get platform input state for this motion source.
 	int InputStateIndex = ControllerMapper.GetInputControllerIndexForMotionSource(MotionSource);
-	if (InputStateIndex != -1)
+	if (InputStateIndex >= 0 && InputStateIndex < MLInput_MaxControllers)
 	{
 		const auto& InputState = InputControllerState[InputStateIndex];
 
