UE-70084 Magic Leap blueprint Functions/Enums fail to load properly when reopening a project
-More modules needed predefault loadingphase.

diff --git a/Engine/Plugins/Lumin/MagicLeap/MagicLeap.uplugin b/Engine/Plugins/Lumin/MagicLeap/MagicLeap.uplugin
--- a/Engine/Plugins/Lumin/MagicLeap/MagicLeap.uplugin
+++ b/Engine/Plugins/Lumin/MagicLeap/MagicLeap.uplugin
@@ -15,7 +15,8 @@
 	"IsBetaVersion" : false,
 	"Installed": false,
 	"SupportedTargetPlatforms": [
-		"Lumin"
+		"Lumin",
+		"Win64"
 	],
 	"Modules" :
 	[
@@ -31,9 +32,10 @@
 			]
 		},
 		{
-			"Name" : "MagicLeapAR",
-			"Type" : "Runtime",
-			"WhitelistPlatforms" :
+			"Name": "MagicLeapAR",
+			"Type": "Runtime",
+			"LoadingPhase": "PreDefault",
+			"WhitelistPlatforms": 
 			[
 				"Win64",
 				"Mac",
@@ -50,9 +52,10 @@
 			]
 		},
 		{
-			"Name" : "MagicLeapController",
-			"Type" : "Runtime",
-			"WhitelistPlatforms" :
+			"Name": "MagicLeapController",
+			"Type": "Runtime",
+			"LoadingPhase": "PreDefault",
+			"WhitelistPlatforms": 
 			[
 				"Win64",
 				"Mac",
@@ -60,9 +63,10 @@
 			]
 		},
 		{
-			"Name" : "MagicLeapEyeTracker",
-			"Type" : "Runtime",
-			"WhitelistPlatforms" :
+			"Name": "MagicLeapEyeTracker",
+			"Type": "Runtime",
+			"LoadingPhase": "PreDefault",
+			"WhitelistPlatforms": 
 			[
 				"Win64",
 				"Mac",
@@ -81,9 +85,10 @@
 			]
 		},
 		{
-			"Name" : "MagicLeapSecureStorage",
-			"Type" : "Runtime",
-			"WhitelistPlatforms" :
+			"Name": "MagicLeapSecureStorage",
+			"Type": "Runtime",
+			"LoadingPhase": "PreDefault",
+			"WhitelistPlatforms": 
 			[
 				"Win64",
 				"Mac",
@@ -111,9 +116,10 @@
 			]
 		},
 		{
-			"Name" : "MagicLeapIdentity",
-			"Type" : "Runtime",
-			"WhitelistPlatforms" :
+			"Name": "MagicLeapIdentity",
+			"Type": "Runtime",
+			"LoadingPhase": "PreDefault",
+			"WhitelistPlatforms": 
 			[
 				"Win64",
 				"Mac",
