UE-70084 Magic Leap blueprint Functions/Enums fail to load properly when reopening a project
-Removed accidental add of win64 to SupportedTargetPlatforms in cl 5296816.

diff --git a/Engine/Plugins/Lumin/MagicLeap/MagicLeap.uplugin b/Engine/Plugins/Lumin/MagicLeap/MagicLeap.uplugin
--- a/Engine/Plugins/Lumin/MagicLeap/MagicLeap.uplugin
+++ b/Engine/Plugins/Lumin/MagicLeap/MagicLeap.uplugin
@@ -15,8 +15,7 @@
 	"IsBetaVersion" : false,
 	"Installed": false,
 	"SupportedTargetPlatforms": [
-		"Lumin",
-		"Win64"
+		"Lumin"
 	],
 	"Modules" :
 	[
