Implement vr.PixelDensity for Lumin

diff --git a/Engine/Plugins/Lumin/MagicLeap/Source/MagicLeap/Private/AppFramework.cpp b/Engine/Plugins/Lumin/MagicLeap/Source/MagicLeap/Private/AppFramework.cpp
--- a/Engine/Plugins/Lumin/MagicLeap/Source/MagicLeap/Private/AppFramework.cpp
+++ b/Engine/Plugins/Lumin/MagicLeap/Source/MagicLeap/Private/AppFramework.cpp
@@ -203,28 +203,6 @@ FVector2D FAppFramework::GetFieldOfView() const
 	return FVector2D(80.0f, 60.0f);
 }
 
-bool FAppFramework::GetDeviceResolution(FVector2D& OutResolution) const
-{
-#if WITH_MLSDK
-	// JMC: The size of the back buffer should not change, rather the internal
-	// viewport should render to a sub region of the back buffer.
-	// This will happen automatically once the render pipeline draws directly
-	// to our provided color buffer.
-	/*const FTrackingFrame *frame = GetOldFrame();
-
-	if (frame && frame->RenderInfoArray.num_virtual_cameras > 0)
-	{
-		const float width = frame->RenderInfoArray.viewport.w * 2.0f;
-		const float height = frame->RenderInfoArray.viewport.h;
-		OutResolution = FVector2D(width, height);
-		return true;
-	}*/
-#endif //WITH_MLSDK
-
-	OutResolution = FVector2D(1280.0f * 2.0f, 960.0f);
-	return true;
-}
-
 uint32 FAppFramework::GetViewportCount() const
 {
 #if WITH_MLSDK
diff --git a/Engine/Plugins/Lumin/MagicLeap/Source/MagicLeap/Private/AppFramework.h b/Engine/Plugins/Lumin/MagicLeap/Source/MagicLeap/Private/AppFramework.h
--- a/Engine/Plugins/Lumin/MagicLeap/Source/MagicLeap/Private/AppFramework.h
+++ b/Engine/Plugins/Lumin/MagicLeap/Source/MagicLeap/Private/AppFramework.h
@@ -66,8 +66,6 @@ public:
 
 	FVector2D GetFieldOfView() const;
 
-	/** Returns @true if we retrieved the device's resolution. Otherwise, it uses the default resolution. */
-	bool GetDeviceResolution(FVector2D& out_resolution) const;
 	FTransform GetDisplayCenterTransform() const { return FTransform::Identity; }; // HACK
 	uint32 GetViewportCount() const;
 
diff --git a/Engine/Plugins/Lumin/MagicLeap/Source/MagicLeap/Private/MagicLeapCustomPresent.cpp b/Engine/Plugins/Lumin/MagicLeap/Source/MagicLeap/Private/MagicLeapCustomPresent.cpp
--- a/Engine/Plugins/Lumin/MagicLeap/Source/MagicLeap/Private/MagicLeapCustomPresent.cpp
+++ b/Engine/Plugins/Lumin/MagicLeap/Source/MagicLeap/Private/MagicLeapCustomPresent.cpp
@@ -195,7 +195,7 @@ void FMagicLeapCustomPresentOpenGL::BeginRendering()
 			UE_LOG(LogMagicLeap, Error, TEXT("MLGraphicsInitFrameParams failed with status %d"), Result);
 		}
 		camera_params.projection_type = MLGraphicsProjectionType_ReversedInfiniteZ;
-		camera_params.surface_scale = frame.ScreenPercentage;
+		camera_params.surface_scale = frame.PixelDensity;
 		camera_params.protected_surface = false;
 		GConfig->GetBool(TEXT("/Script/LuminRuntimeSettings.LuminRuntimeSettings"), TEXT("bProtectedContent"), camera_params.protected_surface, GEngineIni);
 
@@ -317,8 +317,9 @@ void FMagicLeapCustomPresentOpenGL::FinishRendering()
 
 		//check(vp_array.num_virtual_cameras >= 2); // We assume at least one virtual camera per eye
 
-		FVector2D InternalTextureDims;
-		Plugin->GetAppFrameworkConst().GetDeviceResolution(InternalTextureDims);
+		const FIntPoint& IdealRenderTargetSize = Plugin->GetHMDDevice()->GetIdealRenderTargetSize();
+		const int32 SizeX = FMath::CeilToInt(IdealRenderTargetSize.X * Plugin->GetCurrentFrame().PixelDensity);
+		const int32 SizeY = FMath::CeilToInt(IdealRenderTargetSize.Y * Plugin->GetCurrentFrame().PixelDensity);
 
 		// this texture contains both eye renders
 		glBindFramebuffer(GL_FRAMEBUFFER, Framebuffers[0]);
@@ -333,8 +334,8 @@ void FMagicLeapCustomPresentOpenGL::FinishRendering()
 		const bool bShouldFlipVertically = !IsES2Platform(GMaxRHIShaderPlatform);
 
 		bShouldFlipVertically ?
-			FOpenGL::BlitFramebuffer(0, 0, InternalTextureDims.X / 2, InternalTextureDims.Y, 0, vp_height, vp_width, 0, GL_COLOR_BUFFER_BIT, GL_NEAREST) :
-			FOpenGL::BlitFramebuffer(0, 0, InternalTextureDims.X / 2, InternalTextureDims.Y, 0, 0, vp_width, vp_height, GL_COLOR_BUFFER_BIT, GL_NEAREST);
+			FOpenGL::BlitFramebuffer(0, 0, SizeX / 2, SizeY, 0, vp_height, vp_width, 0, GL_COLOR_BUFFER_BIT, GL_NEAREST) :
+			FOpenGL::BlitFramebuffer(0, 0, SizeX / 2, SizeY, 0, 0, vp_width, vp_height, GL_COLOR_BUFFER_BIT, GL_NEAREST);
 
 		MLResult Result = MLGraphicsSignalSyncObjectGL(Plugin->GraphicsClient, vp_array.virtual_cameras[0].sync_object);
 		if (Result != MLResult_Ok)
@@ -345,8 +346,8 @@ void FMagicLeapCustomPresentOpenGL::FinishRendering()
 		FOpenGL::FramebufferTextureLayer(GL_FRAMEBUFFER, GL_COLOR_ATTACHMENT0, vp_array.color_id, 0, 1);
 
 		bShouldFlipVertically ?
-			FOpenGL::BlitFramebuffer(InternalTextureDims.X / 2, 0, InternalTextureDims.X, InternalTextureDims.Y, 0, vp_height, vp_width, 0, GL_COLOR_BUFFER_BIT, GL_NEAREST) :
-			FOpenGL::BlitFramebuffer(InternalTextureDims.X / 2, 0, InternalTextureDims.X, InternalTextureDims.Y, 0, 0, vp_width, vp_height, GL_COLOR_BUFFER_BIT, GL_NEAREST);
+			FOpenGL::BlitFramebuffer(SizeX / 2, 0, SizeX, SizeY, 0, vp_height, vp_width, 0, GL_COLOR_BUFFER_BIT, GL_NEAREST) :
+			FOpenGL::BlitFramebuffer(SizeX / 2, 0, SizeX, SizeY, 0, 0, vp_width, vp_height, GL_COLOR_BUFFER_BIT, GL_NEAREST);
 
 		Result = MLGraphicsSignalSyncObjectGL(Plugin->GraphicsClient, vp_array.virtual_cameras[1].sync_object);
 		if (Result != MLResult_Ok)
@@ -489,7 +490,7 @@ void FMagicLeapCustomPresentVulkan::BeginRendering()
 				UE_LOG(LogMagicLeap, Error, TEXT("MLGraphicsInitFrameParams failed with status %d"), InitResult);
 			}
 			camera_params.projection_type = MLGraphicsProjectionType_UnsignedZ;
-			camera_params.surface_scale = 1.0f;
+			camera_params.surface_scale = RHIframe.PixelDensity;
 			camera_params.protected_surface = false;
 			GConfig->GetBool(TEXT("/Script/LuminRuntimeSettings.LuminRuntimeSettings"), TEXT("bProtectedContent"), camera_params.protected_surface, GEngineIni);
 
diff --git a/Engine/Plugins/Lumin/MagicLeap/Source/MagicLeap/Private/MagicLeapCustomPresent.h b/Engine/Plugins/Lumin/MagicLeap/Source/MagicLeap/Private/MagicLeapCustomPresent.h
--- a/Engine/Plugins/Lumin/MagicLeap/Source/MagicLeap/Private/MagicLeapCustomPresent.h
+++ b/Engine/Plugins/Lumin/MagicLeap/Source/MagicLeap/Private/MagicLeapCustomPresent.h
@@ -48,8 +48,7 @@ public:
 	MLGraphicsVirtualCameraInfoArray RenderInfoArray; // render information for the frame
 #endif //WITH_MLSDK
 
-	// [0.0f, 1.0f]
-	float ScreenPercentage;
+	float PixelDensity;
 	FWorldContext* WorldContext;
 
 	FTrackingFrame()
@@ -70,7 +69,7 @@ public:
 #if WITH_MLSDK
 		, Handle(ML_INVALID_HANDLE)
 #endif //WITH_MLSDK
-		, ScreenPercentage(1.0f)
+		, PixelDensity(1.0f)
 		, WorldContext(nullptr)
 	{
 #if WITH_MLSDK
diff --git a/Engine/Plugins/Lumin/MagicLeap/Source/MagicLeap/Private/MagicLeapHMD.cpp b/Engine/Plugins/Lumin/MagicLeap/Source/MagicLeap/Private/MagicLeapHMD.cpp
--- a/Engine/Plugins/Lumin/MagicLeap/Source/MagicLeap/Private/MagicLeapHMD.cpp
+++ b/Engine/Plugins/Lumin/MagicLeap/Source/MagicLeap/Private/MagicLeapHMD.cpp
@@ -477,14 +477,13 @@ bool FMagicLeapHMD::GetHMDMonitorInfo(MonitorInfo& MonitorDesc)
 	// TODO: check if we can rely on the return bool value from AppFramework.GetDeviceResolution() instead.
 	if (IsInitialized() && bHmdEnabled)
 	{
-		FVector2D Resolution;
-		AppFramework.GetDeviceResolution(Resolution);
+		const FIntPoint& RTSize = GetIdealRenderTargetSize();
 		MonitorDesc.MonitorName = "";
 		MonitorDesc.MonitorId = 0;
 		MonitorDesc.DesktopX = 0;
 		MonitorDesc.DesktopY = 0;
-		MonitorDesc.ResolutionX = Resolution.X;
-		MonitorDesc.ResolutionY = Resolution.Y;
+		MonitorDesc.ResolutionX = RTSize.X;
+		MonitorDesc.ResolutionY = RTSize.Y;
 		return true;
 	}
 	else
@@ -506,6 +505,17 @@ void FMagicLeapHMD::GetFieldOfView(float& OutHFOVInDegrees, float& OutVFOVInDegr
 	OutVFOVInDegrees = frame.VFov;
 }
 
+void FMagicLeapHMD::SetPixelDensity(const float NewDensity)
+{
+	// Surface scale does not support > 1.0
+	PixelDensity = FMath::Clamp(NewDensity, PixelDensityMin, 1.0f);
+}
+
+FIntPoint FMagicLeapHMD::GetIdealRenderTargetSize() const
+{
+	return FIntPoint(1280 * 2, 960);
+}
+
 bool FMagicLeapHMD::DoesSupportPositionalTracking() const
 {
 	return bHmdPosTracking;
@@ -644,8 +654,7 @@ void FMagicLeapHMD::RefreshTrackingFrame()
 #if WITH_MLSDK
 	check(IsInGameThread());
 
-	static const auto CVar = IConsoleManager::Get().FindTConsoleVariableDataFloat(TEXT("r.ScreenPercentage"));
-	GameTrackingFrame.ScreenPercentage = CVar->GetValueOnGameThread() / 100.0f;
+	GameTrackingFrame.PixelDensity = PixelDensity;
 
 	// get the frame id for the frame
 	GameTrackingFrame.FrameId = HeadTrackerData.coord_frame_head;
@@ -890,8 +899,14 @@ void FMagicLeapHMD::AdjustViewRect(EStereoscopicPass StereoPass, int32& X, int32
 {
 	if (DebugViewportWidth > 0)
 	{
-		SizeX = DebugViewportWidth;
-		SizeY = DebugViewportHeight;
+		SizeX = DebugViewportWidth * PixelDensity;
+		SizeY = DebugViewportHeight * PixelDensity;
+	}
+	else
+	{
+		const FIntPoint& IdealRenderTargetSize = GetIdealRenderTargetSize();
+		SizeX = FMath::CeilToInt(IdealRenderTargetSize.X * PixelDensity);
+		SizeY = FMath::CeilToInt(IdealRenderTargetSize.Y * PixelDensity);
 	}
 
 	X = 0;
@@ -990,16 +1005,9 @@ void FMagicLeapHMD::CalculateRenderTargetSize(const class FViewport& Viewport, u
 {
 	check(IsInGameThread());
 
-	// set to the resolution of both eyes. unreal will pass in default 1280x720 which needs to be overridden here.
-	FVector2D Resolution;
-	// Device resolution takes into account the r.ScreenPercentage property that is set as the surface_scale for the camera params on BeginFrame().
-	bool bValidResolution = AppFramework.GetDeviceResolution(Resolution);
-
-	if (bValidResolution)
-	{
-		InOutSizeX = Resolution.X;
-		InOutSizeY = Resolution.Y;
-	}
+	const FIntPoint& IdealRenderTargetSize = GetIdealRenderTargetSize();
+	InOutSizeX = FMath::CeilToInt(IdealRenderTargetSize.X * PixelDensity);
+	InOutSizeY = FMath::CeilToInt(IdealRenderTargetSize.Y * PixelDensity);
 }
 
 bool FMagicLeapHMD::NeedReAllocateViewportRenderTarget(const FViewport& Viewport)
@@ -1080,7 +1088,7 @@ FMagicLeapHMD::FMagicLeapHMD(IMagicLeapPlugin* InMagicLeapPlugin, IARSystemSuppo
 #endif //WITH_MLSDK
 	RendererModule(nullptr),
 	MagicLeapPlugin(InMagicLeapPlugin),
-	IdealScreenPercentage(100.0f),
+	PixelDensity(1.0f),
 	bIsPlaying(false),
 	bIsPerceptionEnabled(false),
 	bIsVDZIEnabled(bEnableVDZI),
@@ -1120,6 +1128,13 @@ void FMagicLeapHMD::Startup()
 	// Context must be created before the bridge so that the bridge can set the render api.
 	AppFramework.Startup();
 
+	// Set initial pixel density
+	static const auto PixelDensityCVar = IConsoleManager::Get().FindConsoleVariable(TEXT("vr.PixelDensity"));
+	if (PixelDensityCVar)
+	{
+		PixelDensity = FMath::Clamp(PixelDensityCVar->GetFloat(), PixelDensityMin, PixelDensityMax);
+	}
+
 #if PLATFORM_WINDOWS
 	if (IsPCPlatform(GMaxRHIShaderPlatform) && !IsOpenGLPlatform(GMaxRHIShaderPlatform) && !IsVulkanPlatform(GMaxRHIShaderPlatform))
 	{
@@ -1448,17 +1463,6 @@ void FMagicLeapHMD::InitDevice()
 		FSceneViewport* SceneVP = FindSceneViewport();
 		if (SceneVP && SceneVP->IsStereoRenderingAllowed())
 		{
-			FVector2D HmdSize;
-			bool IsValidResolution = AppFramework.GetDeviceResolution(HmdSize);
-
-			// An invalid resolution implies that the render thread hasn't received valid render info yet.
-			// If this is the case, we need to defer initialization until that happens.
-			//if (!IsValidResolution)
-			//{
-			//	FPlatformMisc::LowLevelOutputDebugStringf(TEXT("     !IsValidResolution!"));
-			//	return;
-			//}
-
 			// This init must happen on the main thread for VR preview, otherwise it crashes on a non-Lumin RHI.
 
 				// Save any runtime configuration changes from the .ini.
@@ -1471,15 +1475,15 @@ void FMagicLeapHMD::InitDevice()
 			EWindowMode::Type WindowMode = EWindowMode::Windowed;
 #endif
 
-			if (IsValidResolution)
+			if (bHmdEnabled)
 			{
-				// Set the viewport size only if the resolution is valid to have the correct size for standalone desktop builds.
-				DebugViewportWidth = HmdSize.X;
-				DebugViewportHeight = HmdSize.Y;
-				FSystemResolution::RequestResolutionChange(HmdSize.X, HmdSize.Y, WindowMode);
+				const FIntPoint& RTSize = GetIdealRenderTargetSize();
+				DebugViewportWidth = RTSize.X;
+				DebugViewportHeight = RTSize.Y;
+				FSystemResolution::RequestResolutionChange(RTSize.X, RTSize.Y, WindowMode);
 				FPlatformAtomics::InterlockedExchange(&bDeviceWasJustInitialized, 0);
 			}
-			else if (!bHmdEnabled)
+			else
 			{
 				// If HMD is not enabled, set bDeviceWasJustInitialized to false so the device resolution code is not run every frame.
 				FPlatformAtomics::InterlockedExchange(&bDeviceWasJustInitialized, 0);				
diff --git a/Engine/Plugins/Lumin/MagicLeap/Source/MagicLeap/Private/MagicLeapHMD.h b/Engine/Plugins/Lumin/MagicLeap/Source/MagicLeap/Private/MagicLeapHMD.h
--- a/Engine/Plugins/Lumin/MagicLeap/Source/MagicLeap/Private/MagicLeapHMD.h
+++ b/Engine/Plugins/Lumin/MagicLeap/Source/MagicLeap/Private/MagicLeapHMD.h
@@ -103,6 +103,10 @@ public:
 	virtual void UpdateScreenSettings(const FViewport* InViewport) override {}
 	virtual bool IsRenderingPaused() const override { return bIsRenderingPaused; }
 
+	virtual float GetPixelDenity() const override { return PixelDensity; }
+	virtual void SetPixelDensity(const float NewDensity) override;
+	virtual FIntPoint GetIdealRenderTargetSize() const override;
+
 	/** IStereoRendering interface */
 	virtual bool IsStereoEnabled() const override
 	{
@@ -296,7 +300,7 @@ private:
 #endif //WITH_MLSDK
 	IRendererModule* RendererModule;
 	IMagicLeapPlugin* MagicLeapPlugin;
-	float IdealScreenPercentage;
+	float PixelDensity;
 	bool bIsPlaying;
 	bool bIsPerceptionEnabled;
 	bool bIsVDZIEnabled;
