MagicLeap ImageTrackerComponent: Adding check for PLATFORM_LUMIN to prevent PIE crash running code that was designed to only run on device.

This had been fixed previously but somehow the fix got lost merging between streams.

diff --git a/Engine/Plugins/Lumin/MagicLeap/Source/MagicLeap/Private/ImageTrackerComponent.cpp b/Engine/Plugins/Lumin/MagicLeap/Source/MagicLeap/Private/ImageTrackerComponent.cpp
--- a/Engine/Plugins/Lumin/MagicLeap/Source/MagicLeap/Private/ImageTrackerComponent.cpp
+++ b/Engine/Plugins/Lumin/MagicLeap/Source/MagicLeap/Private/ImageTrackerComponent.cpp
@@ -451,7 +451,7 @@ UImageTrackerComponent::~UImageTrackerComponent()
 
 void UImageTrackerComponent::TickComponent(float DeltaTime, enum ELevelTick TickType, FActorComponentTickFunction *ThisTickFunction)
 {
-#if WITH_MLSDK
+#if WITH_MLSDK && PLATFORM_LUMIN
 	Super::TickComponent(DeltaTime, TickType, ThisTickFunction);
 
 	const FAppFramework& AppFramework = StaticCastSharedPtr<FMagicLeapHMD>(GEngine->XRSystem)->GetAppFrameworkConst();
