Merging CL 4304939 from Dev-VR - Commenting out Eye Tracker init log until the feature is complete.

diff --git a/Engine/Plugins/Lumin/MagicLeap/Source/MagicLeapEyeTracker/Private/MagicLeapVREyeTracker.cpp b/Engine/Plugins/Lumin/MagicLeap/Source/MagicLeapEyeTracker/Private/MagicLeapVREyeTracker.cpp
--- a/Engine/Plugins/Lumin/MagicLeap/Source/MagicLeapEyeTracker/Private/MagicLeapVREyeTracker.cpp
+++ b/Engine/Plugins/Lumin/MagicLeap/Source/MagicLeapEyeTracker/Private/MagicLeapVREyeTracker.cpp
@@ -192,7 +192,7 @@ bool FMagicLeapVREyeTracker::Tick(float DeltaTime)
 				bInitialized = CreateResult == MLResult_Ok && MLHandleIsValid(EyeTrackingHandle);
 				if (bInitialized)
 				{
-					UE_LOG(LogCore, Warning, TEXT("   VR Eye Tracker Created"));
+					//UE_LOG(LogCore, Log, TEXT("   VR Eye Tracker Created"));
 					// Needs to be called only once.
 					if (MLEyeTrackingGetStaticData(EyeTrackingHandle, &EyeTrackingStaticData) != MLResult_Ok)
 					{
