Merging CL 4311018 from Dev-VR - Preventing multiple lumin frame timing hints while the graphics context isn't valid

diff --git a/Engine/Plugins/Lumin/MagicLeap/Source/MagicLeap/Private/MagicLeapHMD.cpp b/Engine/Plugins/Lumin/MagicLeap/Source/MagicLeap/Private/MagicLeapHMD.cpp
--- a/Engine/Plugins/Lumin/MagicLeap/Source/MagicLeap/Private/MagicLeapHMD.cpp
+++ b/Engine/Plugins/Lumin/MagicLeap/Source/MagicLeap/Private/MagicLeapHMD.cpp
@@ -1109,7 +1109,10 @@ void FMagicLeapHMD::LoadFromIni()
 	{
 		// This will be set later during Graphics Client initialization based on thje value in FrameTimingHint read from the config here.
 		ConfigFrameTimingHint = static_cast<ELuminFrameTimingHint>(FrameTimingEnum->GetValueByNameString(EnumVal));
-		SetFrameTimingHint(ConfigFrameTimingHint);
+		if (GraphicsClient != ML_INVALID_HANDLE)
+		{
+			SetFrameTimingHint(ConfigFrameTimingHint);
+		}
 	}
 #endif //PLATFORM_LUMIN
 }
