Merging using Dev-VR_to_Release-4.21 Change: 4425273 UEVR-1333 MRMesh normals problems
-Normals should not be world scaled.

Files:
//UE4/Dev-VR/Engine/Plugins/Lumin/MagicLeap/Source/MagicLeap/Private/MeshTrackerComponent.cpp#13

diff --git a/Engine/Plugins/Lumin/MagicLeap/Source/MagicLeap/Private/MeshTrackerComponent.cpp b/Engine/Plugins/Lumin/MagicLeap/Source/MagicLeap/Private/MeshTrackerComponent.cpp
--- a/Engine/Plugins/Lumin/MagicLeap/Source/MagicLeap/Private/MeshTrackerComponent.cpp
+++ b/Engine/Plugins/Lumin/MagicLeap/Source/MagicLeap/Private/MeshTrackerComponent.cpp
@@ -648,7 +648,7 @@ void UMeshTrackerComponent::TickComponent(float DeltaTime, enum ELevelTick TickT
 						for (uint32_t n = 0; n < MeshData.vertex_count; ++ n)
 						{
 							CurrentMeshDataCache->Normals.Add(MagicLeap::
-								ToFVector(MeshData.normal[n], WorldToMetersScale));
+								ToFVector(MeshData.normal[n], 1.0f));
 						}
 					}
 					// If no normals were provided we need to pack fake ones for Vulkan
