Instead of tearing down and spinning up the render thread / rhi thread pipeline on pause, simply disable and enable rendering.
This leaves the rendering infrastructure intact, but with no work to do.

diff --git a/Engine/Plugins/Lumin/MagicLeap/Source/MagicLeap/Private/AppFramework.cpp b/Engine/Plugins/Lumin/MagicLeap/Source/MagicLeap/Private/AppFramework.cpp
--- a/Engine/Plugins/Lumin/MagicLeap/Source/MagicLeap/Private/AppFramework.cpp
+++ b/Engine/Plugins/Lumin/MagicLeap/Source/MagicLeap/Private/AppFramework.cpp
@@ -90,12 +90,6 @@ void FAppFramework::ApplicationPauseDelegate()
 {
 	UE_LOG(LogMagicLeap, Log, TEXT("+++++++ ML AppFramework APP PAUSE ++++++"));
 
-	FMagicLeapHMD * hmd = GEngine ? static_cast<FMagicLeapHMD*>(GEngine->XRSystem->GetHMDDevice()) : nullptr;
-
-	if (hmd && hmd->GetActiveCustomPresent())
-	{
-		hmd->GetActiveCustomPresent()->Reset();
-	}
 
 	if (GEngine)
 	{
@@ -117,16 +111,10 @@ void FAppFramework::ApplicationPauseDelegate()
 	}
 
 	// Pause rendering
-	if (GUseThreadedRendering)
+	FMagicLeapHMD * const HMD = GEngine ? static_cast<FMagicLeapHMD*>(GEngine->XRSystem->GetHMDDevice()) : nullptr;
+	if (HMD)
 	{
-		if (GIsThreadedRendering)
-		{
-			StopRenderingThread(); 
-		}
-	}
-	else
-	{
-		RHIReleaseThreadOwnership();
+		HMD->PauseRendering(true);
 	}
 }
 
@@ -135,16 +123,10 @@ void FAppFramework::ApplicationResumeDelegate()
 	UE_LOG(LogMagicLeap, Log, TEXT("+++++++ ML AppFramework APP RESUME ++++++"));
 
 	// Resume rendering
-	if (GUseThreadedRendering)
-	{
-		if (!GIsThreadedRendering)
-		{
-			StartRenderingThread();
-		}
-	}
-	else
+	FMagicLeapHMD * const HMD = GEngine ? static_cast<FMagicLeapHMD*>(GEngine->XRSystem->GetHMDDevice()) : nullptr;
+	if (HMD)
 	{
-		RHIAcquireThreadOwnership();
+		HMD->PauseRendering(false);
 	}
 
 	if (GEngine)
diff --git a/Engine/Plugins/Lumin/MagicLeap/Source/MagicLeap/Private/MagicLeapHMD.cpp b/Engine/Plugins/Lumin/MagicLeap/Source/MagicLeap/Private/MagicLeapHMD.cpp
--- a/Engine/Plugins/Lumin/MagicLeap/Source/MagicLeap/Private/MagicLeapHMD.cpp
+++ b/Engine/Plugins/Lumin/MagicLeap/Source/MagicLeap/Private/MagicLeapHMD.cpp
@@ -981,6 +981,7 @@ FMagicLeapHMD::FMagicLeapHMD(IMagicLeapPlugin* InMagicLeapPlugin, bool bEnableVD
 	bStereoEnabled(false),
 	bStereoDesired(false),
 #endif
+	bIsRenderingPaused(false),
 	bHmdPosTracking(true),
 	bHaveVisionTracking(false),
 	IPD(0.064f),
diff --git a/Engine/Plugins/Lumin/MagicLeap/Source/MagicLeap/Private/MagicLeapHMD.h b/Engine/Plugins/Lumin/MagicLeap/Source/MagicLeap/Private/MagicLeapHMD.h
--- a/Engine/Plugins/Lumin/MagicLeap/Source/MagicLeap/Private/MagicLeapHMD.h
+++ b/Engine/Plugins/Lumin/MagicLeap/Source/MagicLeap/Private/MagicLeapHMD.h
@@ -83,7 +83,7 @@ public:
 
 	virtual void GetFieldOfView(float& OutHFOVInDegrees, float& OutVFOVInDegrees) const override;
 	virtual void UpdateScreenSettings(const FViewport* InViewport) override {}
-
+	virtual bool IsRenderingPaused() const override { return bIsRenderingPaused; }
 
 	/** IStereoRendering interface */
 	virtual bool IsStereoEnabled() const override
@@ -158,6 +158,7 @@ public:
 	/** Enables, or disables, local input. Returns the previous value of ignore input. */
 	bool SetIgnoreInput(bool Ignore);
 
+	void PauseRendering(const bool bIsPaused) { bIsRenderingPaused = bIsPaused; }
 
 	/** Utility class to scope guard enabling and disabling game viewport client input processing.
 	* On creation it will enable the input processing, and on exit it will restore it to its
@@ -282,6 +283,7 @@ private:
 #if !PLATFORM_LUMIN
 	bool bStereoDesired;
 #endif
+	bool bIsRenderingPaused;
 	bool bHmdPosTracking;
 	mutable bool bHaveVisionTracking;
 	float IPD;
diff --git a/Engine/Source/Runtime/HeadMountedDisplay/Private/DefaultXRCamera.cpp b/Engine/Source/Runtime/HeadMountedDisplay/Private/DefaultXRCamera.cpp
--- a/Engine/Source/Runtime/HeadMountedDisplay/Private/DefaultXRCamera.cpp
+++ b/Engine/Source/Runtime/HeadMountedDisplay/Private/DefaultXRCamera.cpp
@@ -192,10 +192,11 @@ void FDefaultXRCamera::SetupViewFamily(FSceneViewFamily& InViewFamily)
 {
 	static const auto CVarAllowMotionBlurInVR = IConsoleManager::Get().FindTConsoleVariableDataInt(TEXT("vr.AllowMotionBlurInVR"));
 	const bool AllowMotionBlur = (CVarAllowMotionBlurInVR && CVarAllowMotionBlurInVR->GetValueOnAnyThread() != 0);
-	const IHeadMountedDisplay* HMD = TrackingSystem->GetHMDDevice();
+	const IHeadMountedDisplay* const HMD = TrackingSystem->GetHMDDevice();
 	InViewFamily.EngineShowFlags.MotionBlur = AllowMotionBlur;
 	InViewFamily.EngineShowFlags.HMDDistortion = HMD != nullptr ? HMD->GetHMDDistortionEnabled(InViewFamily.Scene->GetShadingPath()) : false;
 	InViewFamily.EngineShowFlags.StereoRendering = bCurrentFrameIsStereoRendering;
+	InViewFamily.EngineShowFlags.Rendering = HMD != nullptr ? !HMD->IsRenderingPaused() : true;
 }
 
 void FDefaultXRCamera::SetupView(FSceneViewFamily& InViewFamily, FSceneView& InView)
diff --git a/Engine/Source/Runtime/HeadMountedDisplay/Public/IHeadMountedDisplay.h b/Engine/Source/Runtime/HeadMountedDisplay/Public/IHeadMountedDisplay.h
--- a/Engine/Source/Runtime/HeadMountedDisplay/Public/IHeadMountedDisplay.h
+++ b/Engine/Source/Runtime/HeadMountedDisplay/Public/IHeadMountedDisplay.h
@@ -242,4 +242,9 @@ public:
 	 * Returns true, if the app has VR focus, meaning if it is rendered in the HMD.
 	 */
 	virtual bool DoesAppHaveVRFocus() const;
+
+	/**
+	 * If true, scene rendering should be skipped.
+	 */
+	virtual bool IsRenderingPaused() const { return false; }
 };
