#DEV-VR - Fix for Desktop Vulkan on Lumin.

diff --git a/Engine/Plugins/Lumin/MagicLeap/Source/MagicLeap/Private/MagicLeapCustomPresent.cpp b/Engine/Plugins/Lumin/MagicLeap/Source/MagicLeap/Private/MagicLeapCustomPresent.cpp
--- a/Engine/Plugins/Lumin/MagicLeap/Source/MagicLeap/Private/MagicLeapCustomPresent.cpp
+++ b/Engine/Plugins/Lumin/MagicLeap/Source/MagicLeap/Private/MagicLeapCustomPresent.cpp
@@ -586,7 +586,7 @@ void FMagicLeapCustomPresentVulkan::FinishRendering()
 				UE_LOG(LogMagicLeap, Log, TEXT("Aliased render target for correct sRGB ouput."));
 			}
 
-			const VkImage FinalTarget = (RenderTargetTextureSRGB != VK_NULL_HANDLE) ? RenderTargetTextureSRGB : RenderTargetTexture;
+			const VkImage FinalTarget = RenderTargetTexture;// (RenderTargetTextureSRGB != VK_NULL_HANDLE) ? RenderTargetTextureSRGB : RenderTargetTexture;
 			FMagicLeapHelperVulkan::BlitImage((uint64)FinalTarget, 0, 0, 0, 0, vp_width, vp_height, 1, (uint64)vp_array.color_id, 0, 0, 0, 0, vp_width, vp_height, 1);
 			FMagicLeapHelperVulkan::BlitImage((uint64)FinalTarget, 0, vp_width, 0, 0, vp_width, vp_height, 1, (uint64)vp_array.color_id, 1, 0, 0, 0, vp_width, vp_height, 1);
 		}
