VR - Fix srgb aliasing on Lumin

diff --git a/Engine/Plugins/Lumin/MagicLeap/Source/MagicLeap/Private/MagicLeapCustomPresent.cpp b/Engine/Plugins/Lumin/MagicLeap/Source/MagicLeap/Private/MagicLeapCustomPresent.cpp
--- a/Engine/Plugins/Lumin/MagicLeap/Source/MagicLeap/Private/MagicLeapCustomPresent.cpp
+++ b/Engine/Plugins/Lumin/MagicLeap/Source/MagicLeap/Private/MagicLeapCustomPresent.cpp
@@ -461,6 +461,7 @@ FMagicLeapCustomPresentVulkan::FMagicLeapCustomPresentVulkan(FMagicLeapHMD* plug
 	: FMagicLeapCustomPresent(plugin)
 	, RenderTargetTexture(VK_NULL_HANDLE)
 	, RenderTargetTextureAllocation(VK_NULL_HANDLE)
+	, RenderTargetTextureAllocationOffset(0)
 	, RenderTargetTextureSRGB(VK_NULL_HANDLE)
 	, LastAliasedRenderTarget(VK_NULL_HANDLE)
 {
@@ -580,13 +581,13 @@ void FMagicLeapCustomPresentVulkan::FinishRendering()
 			if (RenderTargetTextureAllocation != VK_NULL_HANDLE && LastAliasedRenderTarget != RenderTargetTexture)
 			{
 				// TODO: If RenderTargetTextureSRGB is non-null, we're leaking the previous handle here. Also leaking on shutdown.
-				RenderTargetTextureSRGB = (VkImage)FMagicLeapHelperVulkan::AliasImageSRGB((uint64)RenderTargetTextureAllocation, vp_width * 2, vp_height);
+				RenderTargetTextureSRGB = (VkImage)FMagicLeapHelperVulkan::AliasImageSRGB((uint64)RenderTargetTextureAllocation, RenderTargetTextureAllocationOffset, vp_width * 2, vp_height);
 				check(RenderTargetTextureSRGB != VK_NULL_HANDLE);
 				LastAliasedRenderTarget = RenderTargetTexture;
 				UE_LOG(LogMagicLeap, Log, TEXT("Aliased render target for correct sRGB ouput."));
 			}
 
-			const VkImage FinalTarget = RenderTargetTexture;// (RenderTargetTextureSRGB != VK_NULL_HANDLE) ? RenderTargetTextureSRGB : RenderTargetTexture;
+			const VkImage FinalTarget = (RenderTargetTextureSRGB != VK_NULL_HANDLE) ? RenderTargetTextureSRGB : RenderTargetTexture;
 			FMagicLeapHelperVulkan::BlitImage((uint64)FinalTarget, 0, 0, 0, 0, vp_width, vp_height, 1, (uint64)vp_array.color_id, 0, 0, 0, 0, vp_width, vp_height, 1);
 			FMagicLeapHelperVulkan::BlitImage((uint64)FinalTarget, 0, vp_width, 0, 0, vp_width, vp_height, 1, (uint64)vp_array.color_id, 1, 0, 0, 0, vp_width, vp_height, 1);
 		}
@@ -631,6 +632,7 @@ void FMagicLeapCustomPresentVulkan::UpdateViewport(const FViewport& Viewport, FR
 
 	RenderTargetTexture = (VkImage)RT->GetNativeResource();
 	RenderTargetTextureAllocation = reinterpret_cast<FVulkanTexture2D*>(RT->GetTexture2D())->Surface.GetAllocationHandle();
+	RenderTargetTextureAllocationOffset = reinterpret_cast<FVulkanTexture2D*>(RT->GetTexture2D())->Surface.GetAllocationOffset();
 
 	InViewportRHI->SetCustomPresent(this);
 	
diff --git a/Engine/Plugins/Lumin/MagicLeap/Source/MagicLeapHelperVulkan/Private/MagicLeapHelperVulkan.cpp b/Engine/Plugins/Lumin/MagicLeap/Source/MagicLeapHelperVulkan/Private/MagicLeapHelperVulkan.cpp
--- a/Engine/Plugins/Lumin/MagicLeap/Source/MagicLeapHelperVulkan/Private/MagicLeapHelperVulkan.cpp
+++ b/Engine/Plugins/Lumin/MagicLeap/Source/MagicLeapHelperVulkan/Private/MagicLeapHelperVulkan.cpp
@@ -147,7 +147,7 @@ void FMagicLeapHelperVulkan::SignalObjects(uint64 SignalObject0, uint64 SignalOb
 #endif
 }
 
-uint64 FMagicLeapHelperVulkan::AliasImageSRGB(const uint64 Allocation, const uint32 Width, const uint32 Height)
+uint64 FMagicLeapHelperVulkan::AliasImageSRGB(const uint64 Allocation, const uint64 AllocationOffset, const uint32 Width, const uint32 Height)
 {
 #if !PLATFORM_MAC
 	VkImageCreateInfo ImageCreateInfo;
@@ -175,15 +175,12 @@ uint64 FMagicLeapHelperVulkan::AliasImageSRGB(const uint64 Allocation, const uin
 	ImageCreateInfo.samples = VK_SAMPLE_COUNT_1_BIT;
 	ImageCreateInfo.initialLayout = VK_IMAGE_LAYOUT_PREINITIALIZED;
 
-	FVulkanDynamicRHI* const RHI = (FVulkanDynamicRHI*)GDynamicRHI;
-	
+	FVulkanDynamicRHI* RHI = (FVulkanDynamicRHI*)GDynamicRHI;
+	FVulkanDevice* Device = RHI->GetDevice();
 	VkImage Result = VK_NULL_HANDLE;
-	VERIFYVULKANRESULT(VulkanRHI::vkCreateImage(RHI->GetDevice()->GetInstanceHandle(), &ImageCreateInfo, nullptr, &Result));
-
-	// Assuming the offset is 0, we're not currently querying the surface for this.
-	const VkDeviceSize Offset = 0;
+	VERIFYVULKANRESULT(VulkanRHI::vkCreateImage(Device->GetInstanceHandle(), &ImageCreateInfo, nullptr, &Result));
 
-	VERIFYVULKANRESULT(VulkanRHI::vkBindImageMemory(RHI->GetDevice()->GetInstanceHandle(), Result, (VkDeviceMemory)Allocation, Offset));
+	VERIFYVULKANRESULT(VulkanRHI::vkBindImageMemory(Device->GetInstanceHandle(), Result, (VkDeviceMemory)Allocation, AllocationOffset));
 
 	return (uint64)Result;
 #else
diff --git a/Engine/Plugins/Lumin/MagicLeap/Source/MagicLeapHelperVulkan/Public/MagicLeapHelperVulkan.h b/Engine/Plugins/Lumin/MagicLeap/Source/MagicLeapHelperVulkan/Public/MagicLeapHelperVulkan.h
--- a/Engine/Plugins/Lumin/MagicLeap/Source/MagicLeapHelperVulkan/Public/MagicLeapHelperVulkan.h
+++ b/Engine/Plugins/Lumin/MagicLeap/Source/MagicLeapHelperVulkan/Public/MagicLeapHelperVulkan.h
@@ -43,5 +43,5 @@ public:
 	static void BlitImage(uint64 SrcName, int32 SrcLevel, int32 SrcX, int32 SrcY, int32 SrcZ, int32 SrcWidth, int32 SrcHeight, int32 SrcDepth, uint64 DstName, int32 DstLevel, int32 DstX, int32 DstY, int32 DstZ, int32 DstWidth, int32 DstHeight, int32 DstDepth);
 	static void SignalObjects(uint64 SignalObject0, uint64 SignalObject1);
 	static void TestClear(uint64 Dest);
-	static uint64 AliasImageSRGB(const uint64 Allocation, const uint32 Width, const uint32 Height);
+	static uint64 AliasImageSRGB(const uint64 Allocation, const uint64 AllocationOffset, const uint32 Width, const uint32 Height);
 };
diff --git a/Engine/Source/Runtime/VulkanRHI/Public/VulkanResources.h b/Engine/Source/Runtime/VulkanRHI/Public/VulkanResources.h
--- a/Engine/Source/Runtime/VulkanRHI/Public/VulkanResources.h
+++ b/Engine/Source/Runtime/VulkanRHI/Public/VulkanResources.h
@@ -310,6 +310,18 @@ public:
 		}
 	}
 
+	inline uint64 GetAllocationOffset() const
+	{
+		if (ResourceAllocation.IsValid())
+		{
+			return ResourceAllocation->GetOffset();
+		}
+		else
+		{
+			return 0;
+		}
+	}
+
 	FVulkanDevice* Device;
 
 	VkImage Image;
