VR - Mssing file

diff --git a/Engine/Plugins/Lumin/MagicLeap/Source/MagicLeap/Private/MagicLeapCustomPresent.h b/Engine/Plugins/Lumin/MagicLeap/Source/MagicLeap/Private/MagicLeapCustomPresent.h
--- a/Engine/Plugins/Lumin/MagicLeap/Source/MagicLeap/Private/MagicLeapCustomPresent.h
+++ b/Engine/Plugins/Lumin/MagicLeap/Source/MagicLeap/Private/MagicLeapCustomPresent.h
@@ -232,6 +232,7 @@ public:
 protected:
 	VkImage RenderTargetTexture = VK_NULL_HANDLE;
 	VkDeviceMemory RenderTargetTextureAllocation = VK_NULL_HANDLE;
+	uint64 RenderTargetTextureAllocationOffset = 0;
 	VkImage RenderTargetTextureSRGB = VK_NULL_HANDLE;
 	VkImage LastAliasedRenderTarget = VK_NULL_HANDLE;
 };
