UE-59374
Lumin app crashes when using MRMesh on Vulkan Desktop rendering path.
-Vulkan does not work correctly unless we fill in all of the vertex buffers.  MRMesh now requires that via assert and MeshTrackerComponent fills in placeholder data as necessary.

diff --git a/Engine/Plugins/Lumin/MagicLeap/Source/MagicLeap/Private/MeshTrackerComponent.cpp b/Engine/Plugins/Lumin/MagicLeap/Source/MagicLeap/Private/MeshTrackerComponent.cpp
--- a/Engine/Plugins/Lumin/MagicLeap/Source/MagicLeap/Private/MeshTrackerComponent.cpp
+++ b/Engine/Plugins/Lumin/MagicLeap/Source/MagicLeap/Private/MeshTrackerComponent.cpp
@@ -706,27 +706,43 @@ void UMeshTrackerComponent::TickComponent(float DeltaTime, enum ELevelTick TickT
 					}
 
 					// Read normal stream
-					if (Impl->Data.normal_stream_index != MLDataArray_InvalidStreamIndex && CurrentMeshData.streams[Impl->Data.normal_stream_index].type != MLDataArrayType_None)
 					{
-						const uint32 NormalsCount = CurrentMeshData.streams[Impl->Data.normal_stream_index].count;
-						CurrentMeshDataCache->Normals.Reserve(NormalsCount);
-						for (uint32 i = 0; i < NormalsCount; ++i)
+						uint32 NormalsCount = 0;
+						if (Impl->Data.normal_stream_index != MLDataArray_InvalidStreamIndex && CurrentMeshData.streams[Impl->Data.normal_stream_index].type != MLDataArrayType_None)
+						{
+							NormalsCount = CurrentMeshData.streams[Impl->Data.normal_stream_index].count;
+							CurrentMeshDataCache->Normals.Reserve(NormalsCount);
+							for (uint32 i = 0; i < NormalsCount; ++i)
+							{
+								const auto &normal = CurrentMeshData.streams[Impl->Data.normal_stream_index].xyz_array[i];
+								CurrentMeshDataCache->Normals.Add(MagicLeap::ToFVector(normal, 1.0f));
+							}
+						}
+						else
 						{
-							const auto &normal = CurrentMeshData.streams[Impl->Data.normal_stream_index].xyz_array[i];
-							CurrentMeshDataCache->Normals.Add(MagicLeap::ToFVector(normal, 1.0f));
+							// Vulkan requires that tangent data exist, so we provide fake normals and tangents even if it does not
+
+							NormalsCount = CurrentMeshDataCache->Vertices.Num();
+							CurrentMeshDataCache->Normals.Reserve(NormalsCount);
+							for (uint32 i = 0; i < NormalsCount; ++i)
+							{
+								FVector Normal = CurrentMeshDataCache->Vertices[i];
+								Normal.Normalize();
+								CurrentMeshDataCache->Normals.Add(Normal);
+							}
 						}
 
 						// Also write normals into tangents
 						CurrentMeshDataCache->Tangents.Reserve(NormalsCount * 2);
 						for (uint32 i = 0; i < NormalsCount; ++i)
 						{
-							const auto &MLNormal = CurrentMeshData.streams[Impl->Data.normal_stream_index].xyz_array[i];
-							const FVector Normal = MagicLeap::ToFVector(MLNormal, 1.0f);
-							const FVector NonNormal = Normal.X < Normal.Z ? FVector(1,0,0) : FVector(0,1,0);
+							const FVector Normal = CurrentMeshDataCache->Normals[i];
+							const FVector NonNormal = Normal.X < Normal.Z ? FVector(0, 0, 1) : FVector(0, 1, 0);
 							const FVector TangentX = FVector::CrossProduct(Normal, NonNormal);
 							CurrentMeshDataCache->Tangents.Add(TangentX);
 							CurrentMeshDataCache->Tangents.Add(Normal);
 						}
+
 					}
 
 					// Read confidence stream
@@ -741,8 +757,7 @@ void UMeshTrackerComponent::TickComponent(float DeltaTime, enum ELevelTick TickT
 						}
 					}
 
-					// Write VertexColor, if necessary
-					if (VertexColorMode != EMLMeshVertexColorMode::None)
+					// Write VertexColor
 					{
 						switch (VertexColorMode)
 						{
@@ -779,6 +794,16 @@ void UMeshTrackerComponent::TickComponent(float DeltaTime, enum ELevelTick TickT
 							}
 							break;
 						}
+						case EMLMeshVertexColorMode::None:
+						{
+							// Vulkan requires that we fill everything in.
+							const uint32 VertexCount = CurrentMeshDataCache->Vertices.Num();
+							for (uint32 i = 0; i < VertexCount; ++i)
+							{
+								CurrentMeshDataCache->VertexColors.Add(FColor::White);
+							}
+							break;
+						}
 						default:
 							check(false);
 						}
@@ -788,6 +813,16 @@ void UMeshTrackerComponent::TickComponent(float DeltaTime, enum ELevelTick TickT
 					MLDataArrayUnlock(CurrentMeshHandle);
 				}
 
+				// Write UVs
+				{
+					const uint32 VertexCount = CurrentMeshDataCache->Vertices.Num();
+					for (uint32 i = 0; i < VertexCount; ++i)
+					{
+						const float FakeCoord = (float)i / (float)VertexCount;
+						CurrentMeshDataCache->UV0.Add(FVector2D(FakeCoord, FakeCoord));
+					}
+				}
+
 				// We don't add any sections to the procedural mesh component for point clouds.
 				if (MeshType != EMeshType::PointCloud)
 				{
@@ -936,9 +971,6 @@ void UMeshTrackerComponent::TickWithFakeData()
 
 				const int32 VertCount = 8;
 				CurrentMeshDataCache->Vertices.Reserve(VertCount);
-				CurrentMeshDataCache->UV0.Reserve(VertCount);
-				CurrentMeshDataCache->VertexColors.Reserve(VertCount);
-				CurrentMeshDataCache->Triangles.Reserve(6*2*3);
 
 				CurrentMeshDataCache->Vertices.Add(FVector(Origin.X + Extents.X, Origin.Y - Extents.Y, Origin.Z + Extents.Z));  // 0
 				CurrentMeshDataCache->Vertices.Add(FVector(Origin.X + Extents.X, Origin.Y + Extents.Y, Origin.Z + Extents.Z));  // 1
@@ -948,16 +980,36 @@ void UMeshTrackerComponent::TickWithFakeData()
 				CurrentMeshDataCache->Vertices.Add(FVector(Origin.X - Extents.X, Origin.Y + Extents.Y, Origin.Z + Extents.Z));  // 5
 				CurrentMeshDataCache->Vertices.Add(FVector(Origin.X - Extents.X, Origin.Y + Extents.Y, Origin.Z - Extents.Z));  // 6
 				CurrentMeshDataCache->Vertices.Add(FVector(Origin.X - Extents.X, Origin.Y - Extents.Y, Origin.Z - Extents.Z));  // 7
+
+				CurrentMeshDataCache->UV0.Reserve(VertCount);
+				CurrentMeshDataCache->VertexColors.Reserve(VertCount);
 				for (int i = 0; i < VertCount; ++i)
 				{
 					int imod = i % 10;
 					CurrentMeshDataCache->VertexColors.Emplace(FColor::MakeRandomColor());
 					CurrentMeshDataCache->UV0.Add(FVector2D(imod*0.1f, imod*0.1f));
-					//FVector Tangent(imod*0.1f, imod*imod*0.5f, 1.0f);
-					//Tangent.Normalize();
-					//CurrentMeshDataCache->Tangents.Add(Tangent);
 				}
+				
+				// Vulkan requires that tangent data exist
 
+				CurrentMeshDataCache->Normals.Reserve(VertCount);
+				for (uint32 i = 0; i < VertCount; ++i)
+				{
+					CurrentMeshDataCache->Normals.Add(FVector(0,0,1));
+				}
+
+				// Also write normals into tangents
+				CurrentMeshDataCache->Tangents.Reserve(VertCount * 2);
+				for (uint32 i = 0; i < VertCount; ++i)
+				{
+					const FVector Normal = CurrentMeshDataCache->Normals[i];
+					const FVector NonNormal = Normal.X < Normal.Z ? FVector(0, 0, 1) : FVector(0, 1, 0);
+					const FVector TangentX = FVector::CrossProduct(Normal, NonNormal);
+					CurrentMeshDataCache->Tangents.Add(TangentX);
+					CurrentMeshDataCache->Tangents.Add(Normal);
+				}
+				
+				CurrentMeshDataCache->Triangles.Reserve(6 * 2 * 3);
 				// Read triangle (indices) stream
 				// This makes half the triangles of a box, one per side.
 				//CurrentMeshDataCache->Triangles.Add(0);
diff --git a/Engine/Source/Runtime/MRMesh/Private/MRMeshComponent.cpp b/Engine/Source/Runtime/MRMesh/Private/MRMeshComponent.cpp
--- a/Engine/Source/Runtime/MRMesh/Private/MRMeshComponent.cpp
+++ b/Engine/Source/Runtime/MRMesh/Private/MRMeshComponent.cpp
@@ -200,10 +200,11 @@ public:
 		FMRMeshProxySection* NewSection = new FMRMeshProxySection(Args.BrickId, FeatureLevel);
 		ProxySections.Add(NewSection);
 
+		// Vulkan requires that all the buffers be full.
 		const int32 NumVerts = Args.PositionData.Num();
-		check((NumVerts == Args.ColorData.Num()) || Args.ColorData.Num() == 0);
-		check((NumVerts == Args.UVData.Num()) || Args.UVData.Num() == 0);
-		check(((NumVerts * 2) == Args.TangentXZData.Num()) || Args.TangentXZData.Num() == 0);
+		check((NumVerts == Args.ColorData.Num()));
+		check((NumVerts == Args.UVData.Num()));
+		check((NumVerts * 2) == Args.TangentXZData.Num());
 
 		// POSITION BUFFER
 		{
@@ -229,18 +230,6 @@ public:
 			{
 				NewSection->TangentXZBuffer.InitRHIWith(Args.TangentXZData);
 			}
-			else
-			{
-				// Vulkan requires that the tangent data exist, so we fill it in with fake data if necessary.
-				TArray<FPackedNormal> FakeData;
-				FakeData.Reserve(NumVerts * 2);
-				for (int32 i = 0; i < NumVerts; i++)
-				{
-					FakeData.Add(FPackedNormal(FVector::UpVector));
-					FakeData.Add(FPackedNormal(FVector::RightVector));
-				}
-				NewSection->TangentXZBuffer.InitRHIWith(FakeData);
-			}
 
 			if (RHISupportsManualVertexFetch(GMaxRHIShaderPlatform))
 			{
