Enabling RHI threading on Lumin

diff --git a/Engine/Config/BaseDeviceProfiles.ini b/Engine/Config/BaseDeviceProfiles.ini
--- a/Engine/Config/BaseDeviceProfiles.ini
+++ b/Engine/Config/BaseDeviceProfiles.ini
@@ -543,7 +543,6 @@ DeviceType=Android
 +CVars=r.defaultfeature.postprocessing=0
 +CVars=r.Lumin.OverrideExternalTextureSupport=4
 +CVars=r.Vulkan.DelayAcquireBackBuffer=1
-+CVars=r.Vulkan.RHIThread=0
 
 [HTML5 DeviceProfile]
 DeviceType=HTML5
diff --git a/Engine/Plugins/Lumin/MagicLeap/Source/MagicLeap/Private/MagicLeapHMD.cpp b/Engine/Plugins/Lumin/MagicLeap/Source/MagicLeap/Private/MagicLeapHMD.cpp
--- a/Engine/Plugins/Lumin/MagicLeap/Source/MagicLeap/Private/MagicLeapHMD.cpp
+++ b/Engine/Plugins/Lumin/MagicLeap/Source/MagicLeap/Private/MagicLeapHMD.cpp
@@ -88,12 +88,6 @@ public:
 	{
 		FMagicLeapSDKDetection::DetectSDK();
 
-#if PLATFORM_LUMIN
-		// Append the command line environment enclosed within spaces as we can't set it in the
-		// environment. Note that any code accessing GCmdLine before appInit won't
-		// respect the command line environment additions.
-		FCommandLine::Append(TEXT(" \"-norhithread\""));
-#endif //PLATFORM_LUMIN
 
 		// Ideally, we should be able to query GetDefault<UMagicLeapSettings>()->bEnableZI directly.
 		// Unfortunately, the UObject system hasn't finished initialization when this module has been loaded.
