Adding support for DefaultStereoLayers on Lumin

diff --git a/Engine/Plugins/Lumin/MagicLeap/Source/MagicLeap/Private/MagicLeapHMD.h b/Engine/Plugins/Lumin/MagicLeap/Source/MagicLeap/Private/MagicLeapHMD.h
--- a/Engine/Plugins/Lumin/MagicLeap/Source/MagicLeap/Private/MagicLeapHMD.h
+++ b/Engine/Plugins/Lumin/MagicLeap/Source/MagicLeap/Private/MagicLeapHMD.h
@@ -110,6 +110,9 @@ public:
 		return false;
 	}
 
+	IStereoLayers* GetStereoLayers() override;
+
+
 	// FXRRenderTargetManager interface
 	virtual void UpdateViewportRHIBridge(bool bUseSeparateRenderTarget, const class FViewport& Viewport, FRHIViewport* const ViewportRHI) override;
 	virtual bool ShouldUseSeparateRenderTarget() const override
diff --git a/Engine/Plugins/Lumin/MagicLeap/Source/MagicLeap/Private/MagicLeapStereoLayers.cpp b/Engine/Plugins/Lumin/MagicLeap/Source/MagicLeap/Private/MagicLeapStereoLayers.cpp
new file mode 100644
--- /dev/null
+++ b/Engine/Plugins/Lumin/MagicLeap/Source/MagicLeap/Private/MagicLeapStereoLayers.cpp
@@ -0,0 +1,18 @@
+// Copyright 1998-2018 Epic Games, Inc. All Rights Reserved.
+//
+
+#include "MagicLeapStereoLayers.h"
+#include "MagicLeapHMD.h"
+#include "SceneViewExtension.h"
+
+#include "CoreMinimal.h"
+
+IStereoLayers* FMagicLeapHMD::GetStereoLayers()
+{
+	if (!DefaultStereoLayers.IsValid())
+	{
+		TSharedPtr<FMagicLeapStereoLayers, ESPMode::ThreadSafe> NewLayersPtr = FSceneViewExtensions::NewExtension<FMagicLeapStereoLayers>(this);
+		DefaultStereoLayers = StaticCastSharedPtr<FDefaultStereoLayers>(NewLayersPtr);
+	}
+	return DefaultStereoLayers.Get();
+}
diff --git a/Engine/Plugins/Lumin/MagicLeap/Source/MagicLeap/Private/MagicLeapStereoLayers.h b/Engine/Plugins/Lumin/MagicLeap/Source/MagicLeap/Private/MagicLeapStereoLayers.h
new file mode 100644
--- /dev/null
+++ b/Engine/Plugins/Lumin/MagicLeap/Source/MagicLeap/Private/MagicLeapStereoLayers.h
@@ -0,0 +1,12 @@
+// Copyright 1998-2018 Epic Games, Inc. All Rights Reserved.
+
+#pragma once
+#include "CoreMinimal.h"
+
+#include "DefaultStereoLayers.h"
+
+class FMagicLeapStereoLayers : public FDefaultStereoLayers
+{
+public:
+	FMagicLeapStereoLayers(const class FAutoRegister& AutoRegister, class FHeadMountedDisplayBase* InHmd) : FDefaultStereoLayers(AutoRegister, InHmd) {}
+};
