#DEV-VR - Merge from //UE4/Dev-MagicLeap/... @ CL 4136411

diff --git a/Engine/Config/BaseDeviceProfiles.ini b/Engine/Config/BaseDeviceProfiles.ini
--- a/Engine/Config/BaseDeviceProfiles.ini
+++ b/Engine/Config/BaseDeviceProfiles.ini
@@ -543,6 +543,7 @@ DeviceType=Android
 +CVars=r.defaultfeature.postprocessing=0
 +CVars=r.Lumin.OverrideExternalTextureSupport=4
 +CVars=r.Vulkan.DelayAcquireBackBuffer=1
++CVars=vr.DebugCanvasInLayer=1
 
 [HTML5 DeviceProfile]
 DeviceType=HTML5
diff --git a/Engine/Config/BaseEngine.ini b/Engine/Config/BaseEngine.ini
--- a/Engine/Config/BaseEngine.ini
+++ b/Engine/Config/BaseEngine.ini
@@ -1779,8 +1779,20 @@ IconModelPath=(Path="../../Build/Lumin/Resources/Model")
 IconPortalPath=(Path="../../Build/Lumin/Resources/Portal")
 bUseVulkan=true
 VersionCode=1
-bInternetRequired=true
-+AppPrivileges=Headpose
+MinimumOSVersion="lumin_1.0"
++AppPrivileges=LowLatencyLightwear
++AppPrivileges=Internet
++AppPrivileges=LocalAreaNetwork
++AppPrivileges=GesturesSubscribe
++AppPrivileges=GesturesConfig
++AppPrivileges=CameraCapture
++AppPrivileges=WorldReconstruction
++AppPrivileges=PwFoundObjRead
++AppPrivileges=ScreensProvider
++AppPrivileges=AudioCaptureMic
++AppPrivileges=IdentityRead
++AppPrivileges=AudioRecognizer
++AppPrivileges=VoiceInput
 
 [/Script/UnrealEd.CookerSettings]
 DefaultPVRTCQuality=1
diff --git a/Engine/Plugins/Lumin/MagicLeap/MagicLeap.uplugin b/Engine/Plugins/Lumin/MagicLeap/MagicLeap.uplugin
--- a/Engine/Plugins/Lumin/MagicLeap/MagicLeap.uplugin
+++ b/Engine/Plugins/Lumin/MagicLeap/MagicLeap.uplugin
@@ -74,18 +74,17 @@
 			]
 		},
 		{
-			"Name" : "MagicLeapSecureStorage",
+			"Name" : "MagicLeapMusicService",
 			"Type" : "Runtime",
 			"LoadingPhase" : "PostConfigInit",
 			"WhitelistPlatforms" :
 			[
 				"Win64",
 				"Mac",
-				"Lumin"
-			]
+				"Lumin"			]
 		},
 		{
-			"Name" : "MagicLeapHelperOpenGL",
+			"Name" : "MagicLeapSecureStorage",
 			"Type" : "Runtime",
 			"LoadingPhase" : "PostConfigInit",
 			"WhitelistPlatforms" :
@@ -96,7 +95,7 @@
 			]
 		},
 		{
-			"Name" : "MagicLeapHelperVulkan",
+			"Name" : "MagicLeapHelperOpenGL",
 			"Type" : "Runtime",
 			"LoadingPhase" : "PostConfigInit",
 			"WhitelistPlatforms" :
@@ -107,7 +106,7 @@
 			]
 		},
 		{
-			"Name" : "MagicLeapIdentity",
+			"Name" : "MagicLeapHelperVulkan",
 			"Type" : "Runtime",
 			"LoadingPhase" : "PostConfigInit",
 			"WhitelistPlatforms" :
@@ -118,7 +117,7 @@
 			]
 		},
 		{
-			"Name" : "MagicLeapPrivileges",
+			"Name" : "MagicLeapIdentity",
 			"Type" : "Runtime",
 			"LoadingPhase" : "PostConfigInit",
 			"WhitelistPlatforms" :
@@ -127,13 +126,13 @@
 				"Mac",
 				"Lumin"
 			]
-		}
+		},
 	],
 	"Plugins":
 	[
-    {
-      "Name": "ProceduralMeshComponent",
-      "Enabled": true
-    }
+		{
+			"Name": "ProceduralMeshComponent",
+			"Enabled": true
+		}
 	]
 }
diff --git a/Engine/Plugins/Lumin/MagicLeap/Source/MLSDK/MLSDK.Build.cs b/Engine/Plugins/Lumin/MagicLeap/Source/MLSDK/MLSDK.Build.cs
--- a/Engine/Plugins/Lumin/MagicLeap/Source/MLSDK/MLSDK.Build.cs
+++ b/Engine/Plugins/Lumin/MagicLeap/Source/MLSDK/MLSDK.Build.cs
@@ -21,7 +21,6 @@ public class MLSDK : ModuleRules
 		{
 			string IncludePath = Path.Combine(MLSDKPath, "include");
 			string LibraryPath = Path.Combine(MLSDKPath, "lib");
-			string VirtualDeviceLibraryPath = Path.Combine(MLSDKPath, "VirtualDevice", "lib");
 			string LibraryPlatformFolder = string.Empty;
 			switch (Target.Platform)
 			{
@@ -52,11 +51,10 @@ public class MLSDK : ModuleRules
 				if (Target.Platform != UnrealTargetPlatform.Lumin)
 				{
 					string VirtualDeviceIncludePath = "";
-					if (!Ini.TryGetValue("MLSDK", "IncludePath", out VirtualDeviceIncludePath))
+					if (Ini.TryGetValue("MLSDK", "IncludePath", out VirtualDeviceIncludePath))
 					{
-						VirtualDeviceIncludePath = Path.Combine(MLSDKPath, "VirtualDevice", "include");
+						PublicIncludePaths.Add(VirtualDeviceIncludePath);
 					}
-					PublicIncludePaths.Add(VirtualDeviceIncludePath);
 				}
 				PublicIncludePaths.Add(Path.Combine(MLSDKPath, "lumin/usr/include/vulkan"));
 
@@ -69,11 +67,6 @@ public class MLSDK : ModuleRules
 					PublicLibraryPaths.Add(MLSDKLibraryPath);
 				}
 
-				if (Target.Platform != UnrealTargetPlatform.Lumin)
-				{
-					PublicLibraryPaths.Add(VirtualDeviceLibraryPath);
-				}
-
 				string[] MLSDKLibraryList = new string[] {
 					"ml_audio",
 					"ml_camera_metadata",
@@ -97,7 +90,8 @@ public class MLSDK : ModuleRules
 					"ml_privileges",
 					"ml_screens",
 					"ml_secure_storage",
-					"ml_sharedfile"
+					"ml_sharedfile",
+					"ml_purchase"
 				};
 
 				if (Target.Platform == UnrealTargetPlatform.Win64)
@@ -108,8 +102,8 @@ public class MLSDK : ModuleRules
 						PublicDelayLoadDLLs.Add(string.Format("{0}.dll", libname));
 					}
 
-					PublicAdditionalLibraries.Add("ml_virtual_device.lib");
-					PublicDelayLoadDLLs.Add("ml_virtual_device.dll");
+					PublicAdditionalLibraries.Add("ml_remote.lib");
+					PublicDelayLoadDLLs.Add("ml_remote.dll");
 				}
 				else if (Target.Platform == UnrealTargetPlatform.Mac)
 				{
@@ -126,15 +120,11 @@ public class MLSDK : ModuleRules
 						}
 					}
 
-					string virtualDeviceLib = "libml_virtual_device.dylib";
+					string virtualDeviceLib = "libml_remote.dylib";
 					if (!string.IsNullOrEmpty(MLSDKLibraryPath) && File.Exists(Path.Combine(MLSDKLibraryPath, virtualDeviceLib)))
 					{
 						PublicDelayLoadDLLs.Add(Path.Combine(MLSDKLibraryPath, virtualDeviceLib));
 					}
-					else
-					{
-						PublicDelayLoadDLLs.Add(Path.Combine(VirtualDeviceLibraryPath, virtualDeviceLib));
-					}
 				}
 				else if (Target.Platform == UnrealTargetPlatform.Linux)
 				{
@@ -151,15 +141,11 @@ public class MLSDK : ModuleRules
 						}
 					}
 					
-					string virtualDeviceLib = "libml_virtual_device.so";
+					string virtualDeviceLib = "libml_remote.so";
 					if (!string.IsNullOrEmpty(MLSDKLibraryPath) && File.Exists(Path.Combine(MLSDKLibraryPath, virtualDeviceLib)))
 					{
 						PublicDelayLoadDLLs.Add(Path.Combine(MLSDKLibraryPath, virtualDeviceLib));
 					}
-					else
-					{
-						PublicDelayLoadDLLs.Add(Path.Combine(VirtualDeviceLibraryPath, virtualDeviceLib));
-					}
 				}
 				else if (Target.Platform == UnrealTargetPlatform.Lumin)
 				{
diff --git a/Engine/Plugins/Lumin/MagicLeap/Source/MagicLeap/Classes/CameraCaptureComponent.h b/Engine/Plugins/Lumin/MagicLeap/Source/MagicLeap/Classes/CameraCaptureComponent.h
--- a/Engine/Plugins/Lumin/MagicLeap/Source/MagicLeap/Classes/CameraCaptureComponent.h
+++ b/Engine/Plugins/Lumin/MagicLeap/Source/MagicLeap/Classes/CameraCaptureComponent.h
@@ -125,8 +125,7 @@ public:
   FCameraCaptureVidToFileFailure& OnCaptureVidToFileFailure();  
 
 private:
-  class FImageCaptureImpl *Impl;
-  bool bCapturing;
+  class FCameraCaptureImpl *Impl;
 
   UPROPERTY(BlueprintAssignable, Category = "CameraCapture|MagicLeap", meta = (AllowPrivateAccess = true))
   FCameraCaptureLogMessage CaptureLogMessage;
diff --git a/Engine/Plugins/Lumin/MagicLeap/Source/MagicLeap/Classes/ImageTrackerComponent.h b/Engine/Plugins/Lumin/MagicLeap/Source/MagicLeap/Classes/ImageTrackerComponent.h
--- a/Engine/Plugins/Lumin/MagicLeap/Source/MagicLeap/Classes/ImageTrackerComponent.h
+++ b/Engine/Plugins/Lumin/MagicLeap/Source/MagicLeap/Classes/ImageTrackerComponent.h
@@ -8,6 +8,7 @@
 #include "Engine/Texture2D.h"
 #include "Misc/ScopeLock.h"
 #include "HAL/ThreadSafeBool.h"
+#include "Kismet/BlueprintFunctionLibrary.h"
 #include "ImageTrackerComponent.generated.h"
 
 /**
@@ -31,6 +32,21 @@ public:
 	/** Polls for and handles incoming messages from the asynchronous image tracking system. */
 	virtual void TickComponent(float DeltaTime, enum ELevelTick TickType, FActorComponentTickFunction *ThisTickFunction) override;
 
+	/**
+		Attempts to change the currently tracked target.  Initiates an asynchronous call on a worker thread.
+		When the task completes, the instigating blueprint will be notified by either a FSetImageTargetSucceeded
+		or FSetImageTargetFailed event.
+		@param ImageTarget The new texture to be tracked.
+	*/
+	UFUNCTION(BlueprintCallable, Category = "ImageTracking|MagicLeap")
+	void SetTargetAsync(UTexture2D* ImageTarget);
+
+	/** Delegate used to notify the instigating blueprint that the target image was successfully set. */
+	DECLARE_DYNAMIC_MULTICAST_DELEGATE(FSetImageTargetSucceeded);
+
+	/** Delegate used to notify the instigating blueprint that the target image failed to be set. */
+	DECLARE_DYNAMIC_MULTICAST_DELEGATE(FSetImageTargetFailed);
+
 	/** Delegate used to notify the instigating blueprint that the target image is currently visible to the camera */
 	DECLARE_DYNAMIC_MULTICAST_DELEGATE(FImageTargetFound);
 
@@ -46,6 +62,14 @@ public:
 	*/
 	DECLARE_DYNAMIC_MULTICAST_DELEGATE_FourParams(FImageTargetUnreliableTracking, const FVector&, LastTrackedLocation, const FRotator&, LastTrackedRotation, const FVector&, NewUnreliableLocation, const FRotator&, NewUnreliableRotation);
 
+	/** Activated when the target image is successfully set. */
+	UPROPERTY(BlueprintAssignable)
+	FSetImageTargetSucceeded OnSetImageTargetSucceeded;
+
+	/** Activated when the target image fails to be set. */
+	UPROPERTY(BlueprintAssignable)
+	FSetImageTargetFailed OnSetImageTargetFailed;
+
 	/** Activated when the target image becomes visible to the camera */
 	UPROPERTY(BlueprintAssignable)
 	FImageTargetFound OnImageTargetFound;
@@ -91,6 +115,7 @@ public:
 	bool bUseUnreliablePose;
 
 private:
+	friend class FImageTrackerImpl;
 	class FImageTrackerImpl *Impl;
 	bool bTick;
 
@@ -100,3 +125,40 @@ public:
 	void PostEditChangeProperty(FPropertyChangedEvent& PropertyChangedEvent) override;
 #endif
 };
+
+UCLASS(ClassGroup = MagicLeap)
+class MAGICLEAP_API UImageTrackerFunctionLibrary : public UBlueprintFunctionLibrary
+{
+	GENERATED_BODY()
+
+public:
+	/**
+	  Set maximum number of Image Targets that can be tracked at any given time.
+
+	  If the tracker is already tracking the maximum number of targets
+	  possible then it will stop searching for new targets which helps
+	  in reducing the load on the CPU. For example, if you are interested in
+	  tracking a maximum of x targets from a list y (x < y) targets then set this
+	  parameter to x.
+
+	  The valid range for this parameter is from 1 through 25.
+	*/
+	UFUNCTION(BlueprintCallable, Category = "ImageTracking|MagicLeap")
+	static void SetMaxSimultaneousTargets(int32 MaxSimultaneousTargets);
+
+	/**
+	  If true, image tracker will detect and track targets.
+
+	  When enabled Image Tracker will gain access to the camera and start
+	  trackingimages. Enabling image tracker is expensive, takes about 1500ms
+	  on the average.
+
+	  When disabled Image Tracker will release the camera and stop tracking
+	  images. Internal state of the tracker will be maintained (i.e. list of
+	  active/inactive argets and their target_handles).
+
+	  This is done automatically on application pause / resume.
+	*/
+	UFUNCTION(BlueprintCallable, Category = "ImageTracking|MagicLeap")
+	static void EnableImageTracking(bool bEnable);
+};
diff --git a/Engine/Plugins/Lumin/MagicLeap/Source/MagicLeap/Classes/InAppPurchaseComponent.h b/Engine/Plugins/Lumin/MagicLeap/Source/MagicLeap/Classes/InAppPurchaseComponent.h
new file mode 100644
--- /dev/null
+++ b/Engine/Plugins/Lumin/MagicLeap/Source/MagicLeap/Classes/InAppPurchaseComponent.h
@@ -0,0 +1,254 @@
+// %BANNER_BEGIN%
+// ---------------------------------------------------------------------
+// %COPYRIGHT_BEGIN%
+//
+// Copyright (c) 2017 Magic Leap, Inc. (COMPANY) All Rights Reserved.
+// Magic Leap, Inc. Confidential and Proprietary
+//
+// NOTICE: All information contained herein is, and remains the property
+// of COMPANY. The intellectual and technical concepts contained herein
+// are proprietary to COMPANY and may be covered by U.S. and Foreign
+// Patents, patents in process, and are protected by trade secret or
+// copyright law. Dissemination of this information or reproduction of
+// this material is strictly forbidden unless prior written permission is
+// obtained from COMPANY. Access to the source code contained herein is
+// hereby forbidden to anyone except current COMPANY employees, managers
+// or contractors who have executed Confidentiality and Non-disclosure
+// agreements explicitly covering such access.
+//
+// The copyright notice above does not evidence any actual or intended
+// publication or disclosure of this source code, which includes
+// information that is confidential and/or proprietary, and is a trade
+// secret, of COMPANY. ANY REPRODUCTION, MODIFICATION, DISTRIBUTION,
+// PUBLIC PERFORMANCE, OR PUBLIC DISPLAY OF OR THROUGH USE OF THIS
+// SOURCE CODE WITHOUT THE EXPRESS WRITTEN CONSENT OF COMPANY IS
+// STRICTLY PROHIBITED, AND IN VIOLATION OF APPLICABLE LAWS AND
+// INTERNATIONAL TREATIES. THE RECEIPT OR POSSESSION OF THIS SOURCE
+// CODE AND/OR RELATED INFORMATION DOES NOT CONVEY OR IMPLY ANY RIGHTS
+// TO REPRODUCE, DISCLOSE OR DISTRIBUTE ITS CONTENTS, OR TO MANUFACTURE,
+// USE, OR SELL ANYTHING THAT IT MAY DESCRIBE, IN WHOLE OR IN PART.
+//
+// %COPYRIGHT_END%
+// --------------------------------------------------------------------
+// %BANNER_END%
+
+#pragma once
+
+#include "Components/ActorComponent.h"
+#include "Engine/Engine.h"
+#include "InAppPurchaseComponent.generated.h"
+
+UENUM(BlueprintType)
+enum class PurchaseType : uint8
+{
+	/*!
+	\brief This represents an item that can be bought multiple times. The
+	application is responsible for managing the consumption of this item.
+	*/
+	Consumable,
+	/*!
+	\brief This represents an item that can only be bought once. This will
+	be enforced by the services.
+	*/
+	Nonconsumable,
+
+	Undefined,
+};
+
+UENUM(BlueprintType)
+enum class CloudStatus : uint8
+{
+	/*! The current request is still in-progress. */
+	CloudStatus_NotDone = 0,
+	/*! The current request is complete. It may have succeeded or failed. */
+	CloudStatus_Done,
+};
+
+USTRUCT(BlueprintType)
+struct FPurchaseItemDetails
+{
+	GENERATED_BODY()
+
+	/*! This is the id of the item. */
+	FString IAPID;
+	/*! This is the formatted price for the item. */
+	UPROPERTY(EditAnywhere, BlueprintReadOnly, Category = "PurchaseItemDetails|MagicLeap")
+	FString Price;
+	/*! This is the name of the item. */
+	UPROPERTY(EditAnywhere, BlueprintReadOnly, Category = "PurchaseItemDetails|MagicLeap")
+	FString Name;
+	/*! This is the type of purchase. */
+	UPROPERTY(EditAnywhere, BlueprintReadOnly, Category = "PurchaseItemDetails|MagicLeap")
+	PurchaseType Type;
+	/*! This is the token to be used when submitting a purchase. */
+	const char *Token;
+};
+
+USTRUCT(BlueprintType)
+struct FPurchaseConfirmation
+{
+	GENERATED_BODY()
+	/*! This is the unique order id for this purchase. */
+	FString OrderID;
+	/*! This is the name of the item from where this purchase originated. */
+	UPROPERTY(EditAnywhere, BlueprintReadOnly, Category = "PurchaseConfirmation|MagicLeap")
+	FString PackageName;
+	/*! This is the time the product was purchased, in milliseconds since the epoch (Jan 1, 1970). */
+	uint64 PurchaseTime;
+	/*!
+	\brief This is a string containing the signature of the purchase data that
+	was signed with the private key of the developer.
+	*/
+	FString Signature;
+	/*! This is the in app purchase ID of the item being purchased. */
+	FString IAPID;
+	/*! This is the type of purchase. */
+	UPROPERTY(EditAnywhere, BlueprintReadOnly, Category = "PurchaseConfirmation|MagicLeap")
+	PurchaseType Type;
+};
+
+/**
+	The CameraCaptureComponent provides access to and maintains state for camera capture functionality.
+	The connection to the device's camera is managed internally.  Users of this component
+	are able to asynchronously capture camera images and footage to file.  Alternatively,
+	a camera image can be captured directly to texture.  The user need only make the relevant
+	asynchronous call and then register the appropriate success/fail event handlers for the
+	operation's completion.
+*/
+UCLASS(ClassGroup = MagicLeap, BlueprintType, Blueprintable, EditInlineNew, meta = (BlueprintSpawnableComponent))
+class MAGICLEAP_API UInAppPurchaseComponent : public UActorComponent
+{
+	GENERATED_BODY()
+
+public:
+	UInAppPurchaseComponent();
+	virtual ~UInAppPurchaseComponent();
+
+	/** Intializes internal systems. */
+	void BeginPlay() override;
+
+	/** Polls for query results. */
+	virtual void TickComponent(float DeltaTime, enum ELevelTick TickType, FActorComponentTickFunction *ThisTickFunction) override;
+
+	/**
+		Attempts to retrieve details for the specified items.
+		@brief	This call instigates an items details query which is handled on a separate thread.  The result of this asynchronous
+				operation is reported back to the calling blueprint via the FGetItemsDetailsSuccess or FGetItemsDetailsFailure event handlers.
+		@param	ItemIDs An array of FString objects specifying the names of the items whose details will be queried.
+		@return	False if an items details query is already running, true otherwise.
+	*/
+	UFUNCTION(BlueprintCallable, Category = "InAppPurchase|MagicLeap")
+	bool TryGetItemsDetailsAsync(const TArray<FString>& ItemIDs);
+
+	/**
+		Attempts to purchase the specified item.
+		@brief	This call instigates a purchase request which is handled on a separate thread.  The result of this asynchronous operation is
+				reported back to the calling blueprint via the FPurchaseConfirmationSuccess or FPurchaseConfirmationFailure event handlers.
+		@param	ItemDetails The details of the item to be purchased.
+		@return	False if a purchase confirmation is already running, true otherwise.
+	*/
+	UFUNCTION(BlueprintCallable, Category = "InAppPurchase|MagicLeap")
+	bool TryPurchaseItemAsync(const FPurchaseItemDetails& ItemDetails);
+
+	/**
+		Attempts to retrieve the app's purchase history.
+		@brief	This call instigates a purchase history request which is handled on a separate thread.  The result of this asynchronous
+				operation is reported back to the calling blueprint via the FGetPurchaseHistorySuccess or FGetPurchaseHistoryFailure event handlers.
+		@param	InNumPages Specifies the number of history pages to retrieve.
+		@return	False if a purchase history query is already running or InNumPages is less than or equal to zero, true otherwise.
+	*/
+	UFUNCTION(BlueprintCallable, Category = "InAppPurchase|MagicLeap")
+	bool TryGetPurchaseHistoryAsync(int32 InNumPages);
+
+public:
+	/**
+		Delegate used to report log messages.
+		@note This is useful if the user wishes to have log messages in 3D space.
+		@param LogMessage A string containing the log message.
+	*/
+	DECLARE_DYNAMIC_MULTICAST_DELEGATE_OneParam(FInAppPurchaseLogMessage, const FString&, LogMessage);
+
+	/**
+		Delegate used to report a successful retrieval of items details.
+		@param ItemsDetails A list of items details.
+	*/
+	DECLARE_DYNAMIC_MULTICAST_DELEGATE_OneParam(FGetItemsDetailsSuccess, const TArray<FPurchaseItemDetails>&, ItemsDetails);
+
+	/**
+		Delegate used to report a failure to retrieve the requested items details.
+	*/
+	DECLARE_DYNAMIC_MULTICAST_DELEGATE(FGetItemsDetailsFailure);
+
+	/**
+		Delegate used to report a successful item purchase confirmation.
+		@param PurchaseConfirmations A list of item purchase confirmations.
+	*/
+	DECLARE_DYNAMIC_MULTICAST_DELEGATE_OneParam(FPurchaseConfirmationSuccess, const FPurchaseConfirmation&, PurchaseConfirmations);
+
+	/**
+		Delegate used to report a failure to retrieve an item purchase confirmation.
+	*/
+	DECLARE_DYNAMIC_MULTICAST_DELEGATE(FPurchaseConfirmationFailure);
+
+	/**
+		Delegate used to report a successful item purchase history retrieval.
+		@param PurchasesHistory The current purchases history of the app.
+	*/
+	DECLARE_DYNAMIC_MULTICAST_DELEGATE_OneParam(FGetPurchaseHistorySuccess, const TArray<FPurchaseConfirmation>&, PurchaseHistory);
+
+	/**
+		Delegate used to pass a purchase history request failure back to the instigating blueprint.
+	*/
+	DECLARE_DYNAMIC_MULTICAST_DELEGATE(FGetPurchaseHistoryFailure);
+
+	/** Activated when a log message is generated on the capture thread. */
+	FInAppPurchaseLogMessage& OnInAppPurchaseLogMessage();
+
+	/** Activated when retrieval of items details succeeds. */
+	FGetItemsDetailsSuccess& OnGetItemsDetailsSuccess();
+
+	/** Activated when retrieval of items details fails. */
+	FGetItemsDetailsFailure& OnGetItemsDetailsFailure();
+
+	/** Activated when a purchase confirmationn request succeeds. */
+	FPurchaseConfirmationSuccess& OnPurchaseConfirmationSuccess();
+
+	/** Activated when a purchase confirmation request fails. */
+	FPurchaseConfirmationFailure& OnPurchaseConfirmationFailure();
+
+	/** Activated when a purchase history request succeeds. */
+	FGetPurchaseHistorySuccess& OnGetPurchaseHistorySuccess();
+
+	/** Activated when a purchase history request fails. */
+	FGetPurchaseHistoryFailure& OnGetPurchaseHistoryFailure();
+
+private:
+	friend class FInAppPurchaseImpl;
+	class FInAppPurchaseImpl* Impl;
+
+	UPROPERTY(BlueprintAssignable, Category = "InAppPurchase|MagicLeap", meta = (AllowPrivateAccess = true))
+	FInAppPurchaseLogMessage InAppPurchaseLogMessage;
+
+	UPROPERTY(BlueprintAssignable, Category = "InAppPurchase|MagicLeap", meta = (AllowPrivateAccess = true))
+	FGetItemsDetailsSuccess GetItemsDetailsSuccess;
+
+	UPROPERTY(BlueprintAssignable, Category = "InAppPurchase|MagicLeap", meta = (AllowPrivateAccess = true))
+	FGetItemsDetailsFailure GetItemsDetailsFailure;
+
+	UPROPERTY(BlueprintAssignable, Category = "InAppPurchase|MagicLeap", meta = (AllowPrivateAccess = true))
+	FPurchaseConfirmationSuccess PurchaseConfirmationSuccess;
+
+	UPROPERTY(BlueprintAssignable, Category = "InAppPurchase|MagicLeap", meta = (AllowPrivateAccess = true))
+	FPurchaseConfirmationFailure PurchaseConfirmationFailure;
+
+	UPROPERTY(BlueprintAssignable, Category = "InAppPurchase|MagicLeap", meta = (AllowPrivateAccess = true))
+	FGetPurchaseHistorySuccess GetPurchaseHistorySuccess;
+
+	UPROPERTY(BlueprintAssignable, Category = "InAppPurchase|MagicLeap", meta = (AllowPrivateAccess = true))
+	FGetPurchaseHistoryFailure GetPurchaseHistoryFailure;
+
+private:
+	void Log(const FString& LogMessage, bool bError = true);
+};
+
+DECLARE_LOG_CATEGORY_EXTERN(LogInAppPurchase, Verbose, All);
\ No newline at end of file
diff --git a/Engine/Plugins/Lumin/MagicLeap/Source/MagicLeap/Classes/LightingTrackingComponent.h b/Engine/Plugins/Lumin/MagicLeap/Source/MagicLeap/Classes/LightingTrackingComponent.h
new file mode 100644
--- /dev/null
+++ b/Engine/Plugins/Lumin/MagicLeap/Source/MagicLeap/Classes/LightingTrackingComponent.h
@@ -0,0 +1,75 @@
+﻿// %BANNER_BEGIN%
+// ---------------------------------------------------------------------
+// %COPYRIGHT_BEGIN%
+//
+// Copyright (c) 2017 Magic Leap, Inc. (COMPANY) All Rights Reserved.
+// Magic Leap, Inc. Confidential and Proprietary
+//
+// NOTICE: All information contained herein is, and remains the property
+// of COMPANY. The intellectual and technical concepts contained herein
+// are proprietary to COMPANY and may be covered by U.S. and Foreign
+// Patents, patents in process, and are protected by trade secret or
+// copyright law. Dissemination of this information or reproduction of
+// this material is strictly forbidden unless prior written permission is
+// obtained from COMPANY. Access to the source code contained herein is
+// hereby forbidden to anyone except current COMPANY employees, managers
+// or contractors who have executed Confidentiality and Non-disclosure
+// agreements explicitly covering such access.
+//
+// The copyright notice above does not evidence any actual or intended
+// publication or disclosure of this source code, which includes
+// information that is confidential and/or proprietary, and is a trade
+// secret, of COMPANY. ANY REPRODUCTION, MODIFICATION, DISTRIBUTION,
+// PUBLIC PERFORMANCE, OR PUBLIC DISPLAY OF OR THROUGH USE OF THIS
+// SOURCE CODE WITHOUT THE EXPRESS WRITTEN CONSENT OF COMPANY IS
+// STRICTLY PROHIBITED, AND IN VIOLATION OF APPLICABLE LAWS AND
+// INTERNATIONAL TREATIES. THE RECEIPT OR POSSESSION OF THIS SOURCE
+// CODE AND/OR RELATED INFORMATION DOES NOT CONVEY OR IMPLY ANY RIGHTS
+// TO REPRODUCE, DISCLOSE OR DISTRIBUTE ITS CONTENTS, OR TO MANUFACTURE,
+// USE, OR SELL ANYTHING THAT IT MAY DESCRIBE, IN WHOLE OR IN PART.
+//
+// %COPYRIGHT_END%
+// --------------------------------------------------------------------
+// %BANNER_END%
+
+#pragma once
+
+#include "Components/ActorComponent.h"
+#include "Engine/Engine.h"
+#include "LightingTrackingComponent.generated.h"
+
+/**
+The LightingTrackingComponent wraps the Magic Leap lighting tracking API.
+This api provides lumosity data from the camera that can be used to shade objects in a more realistic
+manner (via the post processor).
+*/
+UCLASS(ClassGroup = MagicLeap, BlueprintType, Blueprintable, EditInlineNew, meta = (BlueprintSpawnableComponent))
+class MAGICLEAP_API ULightingTrackingComponent : public UActorComponent
+{
+	GENERATED_BODY()
+
+public:
+	ULightingTrackingComponent();
+
+	/** Intializes the lighting tracking api. If a post processing component cannot be found in the scene, one will be created. */
+	void BeginPlay() override;
+	/** Cleans up the lighting tracking api. */
+	void EndPlay(const EEndPlayReason::Type EndPlayReason) override;
+	/** Polls for data from the camera array and processes it based on the active modes (UseGlobalAmbience, UseColorTemp, ...). */
+	virtual void TickComponent(float DeltaTime, enum ELevelTick TickType, FActorComponentTickFunction *ThisTickFunction) override;
+
+	/** Set to true if you want the global ambience value from the cameras to be used in post processing. */
+	UPROPERTY(EditAnywhere, BlueprintReadWrite, Category = "LightingTracking|MagicLeap")
+		bool UseGlobalAmbience;
+	/** Set to true if you want the color temperature value from the cameras to be used in post processing. */
+	UPROPERTY(EditAnywhere, BlueprintReadWrite, Category = "LightingTracking|MagicLeap")
+		bool UseColorTemp;
+	/** Set to true if you want the ambient cube map to be dynamically updated from the cameras' data. */
+	//UPROPERTY(EditAnywhere, BlueprintReadWrite, Category = "LightingTracking|MagicLeap")
+	//bool UseDynamicAmbientCubeMap;
+
+private:
+	class LightingTrackingImpl* Impl;
+};
+
+DECLARE_LOG_CATEGORY_EXTERN(LogLightingTracking, Display, All);
diff --git a/Engine/Plugins/Lumin/MagicLeap/Source/MagicLeap/Classes/MagicLeapARPinComponent.h b/Engine/Plugins/Lumin/MagicLeap/Source/MagicLeap/Classes/MagicLeapARPinComponent.h
new file mode 100644
--- /dev/null
+++ b/Engine/Plugins/Lumin/MagicLeap/Source/MagicLeap/Classes/MagicLeapARPinComponent.h
@@ -0,0 +1,244 @@
+// %BANNER_BEGIN%
+// ---------------------------------------------------------------------
+// %COPYRIGHT_BEGIN%
+//
+// Copyright (c) 2017 Magic Leap, Inc. (COMPANY) All Rights Reserved.
+// Magic Leap, Inc. Confidential and Proprietary
+//
+// NOTICE: All information contained herein is, and remains the property
+// of COMPANY. The intellectual and technical concepts contained herein
+// are proprietary to COMPANY and may be covered by U.S. and Foreign
+// Patents, patents in process, and are protected by trade secret or
+// copyright law. Dissemination of this information or reproduction of
+// this material is strictly forbidden unless prior written permission is
+// obtained from COMPANY. Access to the source code contained herein is
+// hereby forbidden to anyone except current COMPANY employees, managers
+// or contractors who have executed Confidentiality and Non-disclosure
+// agreements explicitly covering such access.
+//
+// The copyright notice above does not evidence any actual or intended
+// publication or disclosure of this source code, which includes
+// information that is confidential and/or proprietary, and is a trade
+// secret, of COMPANY. ANY REPRODUCTION, MODIFICATION, DISTRIBUTION,
+// PUBLIC PERFORMANCE, OR PUBLIC DISPLAY OF OR THROUGH USE OF THIS
+// SOURCE CODE WITHOUT THE EXPRESS WRITTEN CONSENT OF COMPANY IS
+// STRICTLY PROHIBITED, AND IN VIOLATION OF APPLICABLE LAWS AND
+// INTERNATIONAL TREATIES. THE RECEIPT OR POSSESSION OF THIS SOURCE
+// CODE AND/OR RELATED INFORMATION DOES NOT CONVEY OR IMPLY ANY RIGHTS
+// TO REPRODUCE, DISCLOSE OR DISTRIBUTE ITS CONTENTS, OR TO MANUFACTURE,
+// USE, OR SELL ANYTHING THAT IT MAY DESCRIBE, IN WHOLE OR IN PART.
+//
+// %COPYRIGHT_END%
+// --------------------------------------------------------------------
+// %BANNER_END%
+
+#pragma once
+
+#include "Components/SceneComponent.h"
+#include "UObject/NoExportTypes.h"
+#include "Kismet/BlueprintFunctionLibrary.h"
+#include "MagicLeapARPinComponent.generated.h"
+
+class AActor;
+
+/** List of possible error values for MagicLeapARPin fucntions. */
+UENUM(BlueprintType)
+enum class EPassableWorldError : uint8
+{
+	/** No error. */
+	None,
+	/** Map quality too low for content persistence. Continue building the map. */
+	LowMapQuality,
+	/** Currently unable to localize into any map. Continue building the map. */
+	UnableToLocalize,
+	/** AR Pin is not available at this time. */
+	Unavailable,
+	/** Privileges not met. Add 'PwFoundObjRead' privilege to app manifest. */
+	PrivilegeDenied,
+	/** Invalid function parameter. */
+	InvalidParam,
+	/** Unspecified error. */
+	UnspecifiedFailure
+};
+
+/** Modes for automatically pinning content to real-world. */
+UENUM(BlueprintType)
+enum class EAutoPinType : uint8
+{
+	/** 
+	 * Pin this component / owner actor automatically only if it was pinned in a previous run of the app or replicated over network.
+	 * App needs to call PinSceneComponent() or PinActor() to pin for the very first time.
+	 */
+	OnlyOnDataRestoration,
+	/** Always pin this component / owner actor automatically, without having to call PinSceneComponent() or PinActor() explicitely. */
+	Always,
+	/** Never pin this component / owner actor automatically. App will control pinning and unpinning itself. */
+	Never
+};
+
+/** Direct API interface for the Magic Leap Persistent AR Pin tracker system. */
+UCLASS(ClassGroup = MagicLeap)
+class MAGICLEAP_API UMagicLeapARPinFunctionLibrary : public UBlueprintFunctionLibrary
+{
+	GENERATED_BODY()
+
+public:
+	/**
+	* Returns the count of currently available AR Pins.
+	* @param Count Output param for number of currently available AR Pins. Valid only if return value is EPassableWorldError::None.
+	* @return Error code representing specific success or failure cases.
+	*/
+	UFUNCTION(BlueprintCallable, Category = "ContentPersistence|MagicLeap")
+	static EPassableWorldError GetNumAvailableARPins(int32& Count);
+
+	/**
+	* Returns all the AR Pins currently available.
+	* @param NumRequested Max number of AR Pins to query. Pass in a negative integer to get all available Pins.
+	* @param Pins Output array containing IDs of the found Pins. Valid only if return value is EPassableWorldError::None.
+	* @return Error code representing specific success or failure cases.
+	*/
+	UFUNCTION(BlueprintCallable, Category = "ContentPersistence|MagicLeap")
+	static EPassableWorldError GetAvailableARPins(int32 NumRequested, TArray<FGuid>& Pins);
+
+	/**
+	* Returns the Pin closest to the target point passed in.
+	* @param SearchPoint Position, in world space, to search the closest Pin to.
+	* @param PinID Output param for the ID of the closest Pin. Valid only if return value is EPassableWorldError::None.
+	* @return Error code representing specific success or failure cases.
+	*/
+	UFUNCTION(BlueprintCallable, Category = "ContentPersistence|MagicLeap")
+	static EPassableWorldError GetClosestARPin(const FVector& SearchPoint, FGuid& PinID);
+
+	/**
+	* Returns the world position & orientation of the requested Pin.
+	* @param PinID ID of the Pin to get the position and orientation for.
+	* @param Position Output param for the world position of the Pin. Valid only if return value is true.
+	* @param Orientation Output param for the world orientation of the Pin. Valid only if return value is true.
+	* @return true if the PinID was valid and the position & orientation were successfully retrieved.
+	*/
+	UFUNCTION(BlueprintCallable, Category = "ContentPersistence|MagicLeap")
+	static bool GetARPinPositionAndOrientation(const FGuid& PinID, FVector& Position, FRotator& Orientation);
+
+	/**
+	* [TEMPORARY] Query and cache all available Pins at this time interval.
+	* Function will be removed once the native api adds support for querying if a Pin is available in the current environment.
+	* @return Time interval in seconds to query and cache all available Pins.
+	*/
+	UFUNCTION(BlueprintCallable, BlueprintPure, Category = "ContentPersistence|MagicLeap")
+	static float GetARPinQueryInterval();
+
+	/**
+	* [TEMPORARY] Set the time interval to query and cache all available Pins.
+	* Function will be removed once the native api adds support for querying if a Pin is available in the current environment.
+	* @param QueryInterval Time interval in seconds to query and cache all available Pins.
+	*/
+	UFUNCTION(BlueprintCallable, Category = "ContentPersistence|MagicLeap")
+	static void SetARPinQueryInterval(float QueryInterval);
+};
+
+/** Component to make content persist at locations in the real world. */
+UCLASS(ClassGroup = MagicLeap, BlueprintType, Blueprintable, EditInlineNew, meta = (BlueprintSpawnableComponent))
+class MAGICLEAP_API UMagicLeapARPinComponent : public USceneComponent
+{
+	GENERATED_BODY()
+
+public:
+	UMagicLeapARPinComponent();
+
+	virtual void BeginPlay() override;
+	virtual void TickComponent(float DeltaTime, enum ELevelTick TickType, FActorComponentTickFunction *ThisTickFunction) override;
+
+	/**
+	 * Pin given SceneComponent to the closest AR Pin in real-world.
+	 * OnPersistentEntityPinned event will be fired when a suitable AR Pin is found for this component.
+	 * The component's transform will then be locked. App needs to call UnPin() if it wants to move the component again.
+	 * @param ComponentToPin SceneComponent to pin to the world. Pass in 'this' component if app is using 'OnlyOnDataRestoration' or 'Always' AutoPinType.
+	 * @return true if the component was accepted to be pinned, false otherwise.
+	 */
+	UFUNCTION(BlueprintCallable, Category = "ContentPersistence|MagicLeap")
+	bool PinSceneComponent(USceneComponent* ComponentToPin);
+
+	/**
+	 * Pin given Actor to the closest AR Pin in real-world.
+	 * OnPersistentEntityPinned event will be fired when a suitable AR Pin is found for this Actor.
+	 * The Actor's transform will then be locked. App needs to call UnPin() if it wants to move the Actor again.
+	 * @param ActorToPin Actor to pin to the world. Pass in this component's owner if app is using 'OnlyOnDataRestoration' or 'Always' AutoPinType.
+	 * @return true if the Actor was accepted to be pinned, false otherwise.
+	 */
+	UFUNCTION(BlueprintCallable, Category = "ContentPersistence|MagicLeap")
+	bool PinActor(AActor* ActorToPin);
+
+	/**
+	 * Detach or un-pin the currently pinned entity (component or actor) from the real-world.
+	 * Call this if you want to change the transform of a pinned entity.
+	 * Note that if you still want your content to persist, you will have to call PinSceneComponent() or PinActor() before EndPlay().
+	 */
+	UFUNCTION(BlueprintCallable, Category = "ContentPersistence|MagicLeap")
+	void UnPin();
+
+	/**
+	 * True if an entity (component or actor) is currently pinned by this component.
+	 * If true, the entity's transform will be locked. App needs to call UnPin() if it wants to move it again.
+	 * If false, and you still want your content to persist, you will have to call PinSceneComponent() or PinActor() before EndPlay().
+	 * @return True if an entity (component or actor) is currently pinned by this component.
+	 */
+	UFUNCTION(BlueprintCallable, Category = "ContentPersistence|MagicLeap")
+	bool IsPinned() const;
+
+	/**
+	 * True if the AR Pin for the unique ID ObjectUID was restored from the app's local storage or was repliated over network.
+	 * Implies if content was already pinned earlier. Does not imply if that restored Pin is available in the current environment.
+	 * @return True if the Pin data was restored from local storage or network.
+	 */
+	UFUNCTION(BlueprintCallable, Category = "ContentPersistence|MagicLeap")
+	bool PinRestoredOrSynced() const;
+
+	/**
+	 * Get the ID of the Pin the entity (component or actor) is currently pinned to.
+	 * @param PinID Output param for the ID of the Pin.
+	 * @return True if an entity is currently pinned by this component and the output param is successfully populated.
+	 */
+	UFUNCTION(BlueprintCallable, Category = "ContentPersistence|MagicLeap")
+	bool GetPinnedPinID(FGuid& PinID);
+
+public:
+	/**
+	 * Unique ID for this component to save the meta data for the Pin and make content persistent.
+	 * This name has to be unique across all instances of the MagicLeapARPinComponent class.
+	 * If empty, the name of the owner actor will be used.
+	 */
+	UPROPERTY(BlueprintReadWrite, EditAnywhere, Category = "ContentPersistence|MagicLeap")
+	FString ObjectUID;
+
+	/** Mode for automatically pinning this component or it's owner actor to real-world. */
+	UPROPERTY(BlueprintReadWrite, EditAnywhere, Category = "ContentPersistence|MagicLeap")
+	EAutoPinType AutoPinType;
+
+	/** Pin this component's owner actor instead of just the component itself. Relevant only when using 'OnlyOnDataRestoration' or 'Always' as AutoPinType. */
+	UPROPERTY(BlueprintReadWrite, EditAnywhere, Category = "ContentPersistence|MagicLeap")
+	bool bPinActor;
+
+	/**
+	 * Delegate used to notify the instigating blueprint that an entity (component or actor) has been successfuly pinned to the real-world.
+	 * Indicates that the transform of the pinned entity is now locked. App needs to call UnPin() if it wants to move the entity again.
+	 * @param bRestoredOrSynced True if the entity was pinned as a result of Pin data being restored from local storage or replicatred over network, false if pinned by an explicit PinSceneComponent() or PinActor() call from the app.
+	 */
+	DECLARE_DYNAMIC_MULTICAST_DELEGATE_OneParam(FPersistentEntityPinned, bool, bRestoredOrSynced);
+
+	/** Fired when an entity is successfully pinned by this component. */
+	UPROPERTY(BlueprintAssignable)
+	FPersistentEntityPinned OnPersistentEntityPinned;
+
+private:
+	UPROPERTY()
+	FGuid PinnedCFUID;
+
+	UPROPERTY()
+	USceneComponent* PinnedSceneComponent;
+
+	FTransform ComponentRelativeToCFUID;
+	bool bPinned;
+	bool bDataRestored;
+
+	class FMagicLeapARPinTrackerImpl *Impl;
+};
diff --git a/Engine/Plugins/Lumin/MagicLeap/Source/MagicLeap/Classes/MeshTrackerComponent.h b/Engine/Plugins/Lumin/MagicLeap/Source/MagicLeap/Classes/MeshTrackerComponent.h
--- a/Engine/Plugins/Lumin/MagicLeap/Source/MagicLeap/Classes/MeshTrackerComponent.h
+++ b/Engine/Plugins/Lumin/MagicLeap/Source/MagicLeap/Classes/MeshTrackerComponent.h
@@ -9,15 +9,13 @@
 UENUM(BlueprintType)
 enum class EMeshType : uint8
 {
-	/** The full, high resolution mesh. */
-	Full,
-	/*! Meshing should be done in blocks. */
-	Blocks,
+	/*! Meshing should be done as triangles. */
+	Triangles,
 	/*! Return mesh vertices as a point cloud. */
 	PointCloud
 };
 
-/** Type of mesh to query from the underlying system. */
+/** Vertex color mode. */
 UENUM(BlueprintType)
 enum class EMLMeshVertexColorMode : uint8
 {
@@ -29,6 +27,18 @@ enum class EMLMeshVertexColorMode : uint8
 	Block  UMETA(DisplayName = "Blocks Colored")
 };
 
+/** Discrete level of detail required. */
+UENUM(BlueprintType)
+enum class EMeshLOD : uint8
+{
+	/*! Minimum LOD. */
+	Minimum,
+	/*! Medium LOD. */
+	Medium,
+	/*! Maximum LOD. */
+	Maximum,
+};
+
 /**
 	The MeshTrackerComponent class manages requests for environmental mesh data, processes the results and provides
 	them to the calling system. The calling system is able request environmental mesh data within a specified area.
@@ -80,54 +90,29 @@ public:
 
 	/** The type of mesh to query. */
 	UPROPERTY(EditAnywhere, BlueprintReadWrite, Category = "Meshing|MagicLeap")
-	EMeshType MeshType = EMeshType::Blocks;
-
-	/**
-	  Specifies the time (in seconds) to query if new mesh data is available.
-	  Decreasing this time may give more mesh data, but may degrade performance. Use 0 to turn off auto updating.
-	*/
-	UPROPERTY(EditAnywhere, BlueprintReadWrite, Category = "Meshing|MagicLeap")
-	float MeshingPollTime = 1.0f;
+	EMeshType MeshType = EMeshType::Triangles;
 
 	/** Bounding box for the mesh scan. The mesh will be scanned for only within this box. */
 	UPROPERTY(EditAnywhere, BlueprintReadWrite, Category = "Meshing|MagicLeap")
 	class UBoxComponent* BoundingVolume;
 
-	/** Mesh everything within range. */
+	/** Meshing LOD. */
 	UPROPERTY(EditAnywhere, BlueprintReadWrite, Category = "Meshing|MagicLeap")
-	bool IgnoreBoundingVolume = false;
+	EMeshLOD LevelOfDetail = EMeshLOD::Medium;
 
-	/**
-	  Maximum number of triangles allowed for the mesh in Full mode, or per mesh block in Blocks mode. Use 0 to turn off simplification.
-	  Setting this to a reasonably high enough number because of target hardware limitation. It is recommended to leave this number as low as possible but never 0.
-	*/
-	UPROPERTY(EditAnywhere, BlueprintReadWrite, Category = "Meshing|MagicLeap", meta = (ClampMin = 0))
-	int32 TargetNumberTriangles = 5000;
-
-	/** If true, the system will attempt to fill small gaps to create a more connected mesh. */
-	UPROPERTY(EditAnywhere, BlueprintReadWrite, Category = "Meshing|MagicLeap")
-	bool FillGaps = true;
-
-	/** The perimeter (in Unreal Units) of gaps to be filled. */
+	/** The perimeter (in Unreal Units) of gaps to be filled. 0 means do not fill. A good value is 300cm. */
 	UPROPERTY(EditAnywhere, BlueprintReadWrite, Category = "Meshing|MagicLeap")
-	float PerimeterOfGapsToFill = 100.0f;
+	float PerimeterOfGapsToFill = 300.0f;
 
 	/** If true, the system will planarize the returned mesh i.e. planar regions will be smoothed out. */
 	UPROPERTY(EditAnywhere, BlueprintReadWrite, Category = "Meshing|MagicLeap")
 	bool Planarize = false;
 
-	/** If true, the system will remove mesh sections that are disconnected from the main mesh
-		and which have an area less than DisconnectedSectionArea. */
+	/** Any section that is disconnected from the main mesh and has an area (in Unreal Units squared) 
+	    less than this value will be removed. 
+		0 means do not remove disconnected sections. A good value is 50cm. */
 	UPROPERTY(EditAnywhere, BlueprintReadWrite, Category = "Meshing|MagicLeap")
-	bool RemoveDisconnectedSections = true;
-
-	/** Any section that is disconnected from the main mesh and has an area (in Unreal Units squared) less than this value will be removed. */
-	UPROPERTY(EditAnywhere, BlueprintReadWrite, Category = "Meshing|MagicLeap")
-	float DisconnectedSectionArea = 25.0f;
-
-	/** Minimum distance that a bounding box has to move to cause a rescan. */
-	UPROPERTY(EditAnywhere, BlueprintReadWrite, Category = "Meshing|MagicLeap", meta = (ClampMin = 0.0f))
-	float MinDistanceRescan = 50.0f;
+	float DisconnectedSectionArea = 50.0f;
 
 	/** If true, the system will generate normals for the triangle vertices. */
 	UPROPERTY(EditAnywhere, BlueprintReadWrite, Category = "Meshing|MagicLeap")
@@ -175,10 +160,6 @@ public:
 		It seems like those are not being respected directly by the component.
 	*/
 
-	/** Force an update on a non-live updating mesh tracker. */
-	UFUNCTION(BlueprintCallable, Category = "Meshing|MagicLeap")
-	bool ForceMeshUpdate();
-
 	/** Polls for and handles the results of the environmental mesh queries. */
 	virtual void TickComponent(float DeltaTime, enum ELevelTick TickType, FActorComponentTickFunction *ThisTickFunction) override;
 
@@ -193,8 +174,6 @@ public:
 #endif
 
 private:
-	void LogMLDataArray(const struct MLDataArray& Data) const;
-	void TickWithFakeData();
 	class FMeshTrackerImpl *Impl;
 
 #if WITH_EDITOR
diff --git a/Engine/Plugins/Lumin/MagicLeap/Source/MagicLeap/Classes/RaycastComponent.h b/Engine/Plugins/Lumin/MagicLeap/Source/MagicLeap/Classes/RaycastComponent.h
--- a/Engine/Plugins/Lumin/MagicLeap/Source/MagicLeap/Classes/RaycastComponent.h
+++ b/Engine/Plugins/Lumin/MagicLeap/Source/MagicLeap/Classes/RaycastComponent.h
@@ -53,7 +53,7 @@ enum class ERaycastResultState : uint8
 	/** The raycast request failed. */
 	RequestFailed,
 	/** The ray passed beyond maximum raycast distance and it doesn't hit any surface. */
-	NoCollission,
+	NoCollision,
 	/** The ray hit an unobserved area. This can occur only when CollideWithUnobserved is set to true. */
 	HitUnobserved,
 	/** The ray hit an observed area. */
diff --git a/Engine/Plugins/Lumin/MagicLeap/Source/MagicLeap/MagicLeap.Build.cs b/Engine/Plugins/Lumin/MagicLeap/Source/MagicLeap/MagicLeap.Build.cs
--- a/Engine/Plugins/Lumin/MagicLeap/Source/MagicLeap/MagicLeap.Build.cs
+++ b/Engine/Plugins/Lumin/MagicLeap/Source/MagicLeap/MagicLeap.Build.cs
@@ -50,7 +50,8 @@ namespace UnrealBuildTool.Rules
 					"UtilityShaders",
 					// Public headers of MagicLeapHelperVulkan are protected against Mac so this is fine here.
 					"MagicLeapHelperVulkan",
-					"LuminRuntimeSettings"
+					"LuminRuntimeSettings",
+					"MagicLeapSecureStorage"
 				}
 			);
 
diff --git a/Engine/Plugins/Lumin/MagicLeap/Source/MagicLeap/Private/AppEventHandler.cpp b/Engine/Plugins/Lumin/MagicLeap/Source/MagicLeap/Private/AppEventHandler.cpp
--- a/Engine/Plugins/Lumin/MagicLeap/Source/MagicLeap/Private/AppEventHandler.cpp
+++ b/Engine/Plugins/Lumin/MagicLeap/Source/MagicLeap/Private/AppEventHandler.cpp
@@ -7,8 +7,65 @@
 
 namespace MagicLeap
 {
+#if WITH_MLSDK
+	const TCHAR* MLPrivilegeToString(MLPrivilegeID PrivilegeID)
+	{
+		#define PRIV_TO_STR_CASE(x) case x: { return UTF8_TO_TCHAR((#x)); }
+		switch (PrivilegeID)
+		{
+		PRIV_TO_STR_CASE(MLPrivilegeID_Invalid)
+		PRIV_TO_STR_CASE(MLPrivilegeID_AudioRecognizer)
+		PRIV_TO_STR_CASE(MLPrivilegeID_BatteryInfo)
+		PRIV_TO_STR_CASE(MLPrivilegeID_CameraCapture)
+		PRIV_TO_STR_CASE(MLPrivilegeID_WorldReconstruction)
+		PRIV_TO_STR_CASE(MLPrivilegeID_InAppPurchase)
+		PRIV_TO_STR_CASE(MLPrivilegeID_AudioCaptureMic)
+		PRIV_TO_STR_CASE(MLPrivilegeID_DrmCertificates)
+		PRIV_TO_STR_CASE(MLPrivilegeID_Occlusion)
+		PRIV_TO_STR_CASE(MLPrivilegeID_LowLatencyLightwear)
+		PRIV_TO_STR_CASE(MLPrivilegeID_Internet)
+		PRIV_TO_STR_CASE(MLPrivilegeID_IdentityRead)
+		PRIV_TO_STR_CASE(MLPrivilegeID_BackgroundDownload)
+		PRIV_TO_STR_CASE(MLPrivilegeID_BackgroundUpload)
+		PRIV_TO_STR_CASE(MLPrivilegeID_MediaDrm)
+		PRIV_TO_STR_CASE(MLPrivilegeID_Media)
+		PRIV_TO_STR_CASE(MLPrivilegeID_MediaMetadata)
+		PRIV_TO_STR_CASE(MLPrivilegeID_PowerInfo)
+		PRIV_TO_STR_CASE(MLPrivilegeID_LocalAreaNetwork)
+		PRIV_TO_STR_CASE(MLPrivilegeID_VoiceInput)
+		PRIV_TO_STR_CASE(MLPrivilegeID_Documents)
+		PRIV_TO_STR_CASE(MLPrivilegeID_ConnectBackgroundMusicService)
+		PRIV_TO_STR_CASE(MLPrivilegeID_RegisterBackgroundMusicService)
+		PRIV_TO_STR_CASE(MLPrivilegeID_PwFoundObjRead)
+		PRIV_TO_STR_CASE(MLPrivilegeID_NormalNotificationsUsage)
+		PRIV_TO_STR_CASE(MLPrivilegeID_MusicService)
+		PRIV_TO_STR_CASE(MLPrivilegeID_ControllerPose)
+		PRIV_TO_STR_CASE(MLPrivilegeID_ScreensProvider)
+		PRIV_TO_STR_CASE(MLPrivilegeID_GesturesSubscribe)
+		PRIV_TO_STR_CASE(MLPrivilegeID_GesturesConfig)
+		default: UE_LOG(LogMagicLeap, Error, TEXT("Unmapped privilege %d"), static_cast<int32>(PrivilegeID));
+		}
+
+		return UTF8_TO_TCHAR("");
+	}
+
+	IAppEventHandler::IAppEventHandler(TArray<MLPrivilegeID> InRequiredPrivilegeIDs)
+	: bAllPrivilegesInSync(false)
+	, bWasSystemEnabledOnPause(false)
+	{
+		RequiredPrivileges.Reserve(InRequiredPrivilegeIDs.Num());
+		for (MLPrivilegeID RequiredPrivilegeID : InRequiredPrivilegeIDs)
+		{
+			RequiredPrivileges.Add(RequiredPrivilegeID, FRequiredPrivilege(RequiredPrivilegeID));
+		}
+
+		FAppFramework::AddEventHandler(this);
+	}
+#endif // WITH_MLSDK
+
 	IAppEventHandler::IAppEventHandler()
-	:  bWasSystemEnabledOnPause(false)
+	: bAllPrivilegesInSync(true)
+	, bWasSystemEnabledOnPause(false)
 	{
 		FAppFramework::AddEventHandler(this);
 	}
@@ -18,6 +75,117 @@ namespace MagicLeap
 		FAppFramework::RemoveEventHandler(this);
 	}
 
+#if WITH_MLSDK
+	EPrivilegeState IAppEventHandler::GetPrivilegeStatus(MLPrivilegeID PrivilegeID, bool bBlocking/* = true*/)
+	{
+		FScopeLock Lock(&CriticalSection);
+		FRequiredPrivilege& RequiredPrivilege = RequiredPrivileges[PrivilegeID];
+		if (RequiredPrivilege.State == EPrivilegeState::NotYetRequested)
+		{
+			if (bBlocking)
+			{
+				MLResult Result = MLPrivilegesRequestPrivilege(PrivilegeID);
+				switch (Result)
+				{
+				case MLPrivilegesResult_Granted:
+				{
+					RequiredPrivilege.State = EPrivilegeState::Granted;
+					UE_LOG(LogMagicLeap, Error, TEXT("Privilege '%s' was granted."), MLPrivilegeToString(RequiredPrivilege.PrivilegeID));
+				}
+				break;
+				case MLPrivilegesResult_Denied:
+				{
+					RequiredPrivilege.State = EPrivilegeState::Denied;
+					UE_LOG(LogMagicLeap, Error, TEXT("Privilege '%s' was denied."), MLPrivilegeToString(RequiredPrivilege.PrivilegeID));
+				}
+				break;
+				default:
+				{
+					UE_LOG(LogMagicLeap, Error, TEXT("MLPrivilegesRequestPrivilege() failed with error %s"), UTF8_TO_TCHAR(MLPrivilegesGetResultString(Result)));
+				}
+				}
+			}
+			else
+			{
+				MLPrivilegesRequestPrivilegeAsync(PrivilegeID, &RequiredPrivilege.PrivilegeRequest);
+				RequiredPrivilege.State = EPrivilegeState::Pending;
+			}
+		}
+
+		return RequiredPrivilege.State;
+	}
+#endif // WITH_MLSDK
+
+	void IAppEventHandler::OnAppTick()
+	{
+#if WITH_MLSDK
+		FScopeLock Lock(&CriticalSection);
+		if (bAllPrivilegesInSync)
+		{
+			return;
+		}
+
+		bAllPrivilegesInSync = true;
+
+		for(auto& ItRequiredPrivilege : RequiredPrivileges)
+		{
+			FRequiredPrivilege& RequiredPrivilege = ItRequiredPrivilege.Value;
+			if (RequiredPrivilege.State == EPrivilegeState::NotYetRequested)
+			{
+				bAllPrivilegesInSync = false;
+				continue;
+			}
+
+			if (RequiredPrivilege.State == EPrivilegeState::Granted || RequiredPrivilege.State == EPrivilegeState::Denied)
+			{
+				continue;
+			}
+
+			if (RequiredPrivilege.State == EPrivilegeState::Pending)
+			{
+				MLResult Result = MLPrivilegesRequestPrivilegeTryGet(RequiredPrivilege.PrivilegeRequest);
+				switch (Result)
+				{
+				case MLPrivilegesResult_Granted:
+				{
+					RequiredPrivilege.State = EPrivilegeState::Granted;
+					UE_LOG(LogMagicLeap, Error, TEXT("Privilege '%s' was granted."), MLPrivilegeToString(RequiredPrivilege.PrivilegeID));
+				}
+				break;
+				case MLPrivilegesResult_Denied:
+				{
+					RequiredPrivilege.State = EPrivilegeState::Denied;
+					UE_LOG(LogMagicLeap, Error, TEXT("Privilege '%s' was denied."), MLPrivilegeToString(RequiredPrivilege.PrivilegeID));
+				}
+				break;
+				case MLResult_Pending:
+				{
+					bAllPrivilegesInSync = false;
+				}
+				break;
+				default:
+				{
+					bAllPrivilegesInSync = false;
+					UE_LOG(LogMagicLeap, Error, TEXT("MLPrivilegesRequestPrivilegeTryGet() failed with error %s"), UTF8_TO_TCHAR(MLPrivilegesGetResultString(Result)));
+				}
+				}
+			}
+		}
+#endif // WITH_MLSDK
+	}
+
+	void IAppEventHandler::OnAppResume()
+	{
+#if WITH_MLSDK
+		FScopeLock Lock(&CriticalSection);
+		bAllPrivilegesInSync = false;
+		for (auto& ItRequiredPrivilege : RequiredPrivileges)
+		{
+			ItRequiredPrivilege.Value.State = EPrivilegeState::NotYetRequested;
+		}
+#endif // WITH_MLSDK
+	}
+
 	void IAppEventHandler::AsyncDestroy()
 	{
 		FAppFramework::AsyncDestroy(this);
diff --git a/Engine/Plugins/Lumin/MagicLeap/Source/MagicLeap/Private/AppEventHandler.h b/Engine/Plugins/Lumin/MagicLeap/Source/MagicLeap/Private/AppEventHandler.h
--- a/Engine/Plugins/Lumin/MagicLeap/Source/MagicLeap/Private/AppEventHandler.h
+++ b/Engine/Plugins/Lumin/MagicLeap/Source/MagicLeap/Private/AppEventHandler.h
@@ -3,9 +3,35 @@
 #pragma once
 
 #include "CoreMinimal.h"
+#if WITH_MLSDK
+#include "ml_privileges.h"
+#endif //WITH_MLSDK
 
 namespace MagicLeap
 {
+	enum EPrivilegeState
+	{
+		NotYetRequested,
+		Pending,
+		Granted,
+		Denied
+	};
+
+	struct FRequiredPrivilege
+	{
+#if WITH_MLSDK
+		FRequiredPrivilege(MLPrivilegeID InPrivilegeID)
+		: PrivilegeID(InPrivilegeID)
+		, PrivilegeRequest(nullptr)
+		, State(NotYetRequested)
+		{}
+
+		MLPrivilegeID PrivilegeID;
+		MLPrivilegesAsyncRequest* PrivilegeRequest;
+#endif // WITH_MLSDK
+		EPrivilegeState State;
+	};
+
 	/**
 	* Provides an interface between the AppFramework and any system that needs to be
 	* notified of application events (such as pause/resume).
@@ -13,18 +39,20 @@ namespace MagicLeap
 	class MAGICLEAP_API IAppEventHandler
 	{
 	public:
-		/** Adds the IAppEventHandler instance to the application's list of IAppEventHandler instances.*/
+#if WITH_MLSDK
+		/** 
+			Adds the IAppEventHandler instance to the application's list of IAppEventHandler instances.
+			Populates a RequiredPrivileges list based on the privilge ids passed via InRequiredPrivileges;
+			@param InRequiredPrivileges The list of privilge ids required by the calling system.
+		*/
+		IAppEventHandler(TArray<MLPrivilegeID> InRequiredPrivileges);
+#endif // WITH_MLSDK
+
 		IAppEventHandler();
 
 		/** Removes the IAppEventHandler instance from the application's list of IAppEventHandler instances.*/
 		virtual ~IAppEventHandler();
 
-		/**
-			Can be overridden by inheriting class that needs to perform certain initializations after the app is finished
-			booting up.  AppFramework remains agnostic of the subsystem being initialized.
-		*/
-		virtual void OnAppStartup() {}
-
 		/**
 			Can be overridden by inheriting class that needs to destroy certain api interfaces before the perception stack is
 			closed down.
@@ -32,24 +60,41 @@ namespace MagicLeap
 		virtual void OnAppShutDown() {}
 
 		/**
-			Can be overridden by inheriting class in order to pause its system.  AppFramework
-			remains agnostic of the subsystem being paused.
+			Use to check status of privilege requests.
+		*/
+		virtual void OnAppTick();
+
+		/**
+			Can be overridden by inheriting class in order to pause its system.
 		*/
 		virtual void OnAppPause() {}
 
 		/**
-			Can be overridden by inheriting class in order to resume its system.  AppFramework
-			remains agnostic of the subsystem being resumed.
+			Can be overridden by inheriting class in order to resume its system.
+		*/
+		virtual void OnAppResume();
+
+#if WITH_MLSDK
+		/**
+			Returns the status of the specified privilege.
+			@param PrivilegeID The privilege id to be queried.
+			@param bBlocking Flags whether or not to use the blocking query internally.
 		*/
-		virtual void OnAppResume() {}
+		EPrivilegeState GetPrivilegeStatus(MLPrivilegeID PrivilegeID, bool bBlocking = true);
+#endif //WITH_MLSDK
 
 		/**
 			Pushes this object onto a worker thread so that it's blocking destructor can be called without locking up the update thread.
-			@note This should only be called by objects that have an asynchronous destructor and are no longer referenced.
+			@note This should only be called by objects that have a blocking destructor and are no longer referenced.
 		*/
 		void AsyncDestroy();
 
 	protected:
+#if WITH_MLSDK
+		TMap<MLPrivilegeID, FRequiredPrivilege> RequiredPrivileges;
+#endif // WITH_MLSDK
+		bool bAllPrivilegesInSync;
 		bool bWasSystemEnabledOnPause;
+		FCriticalSection CriticalSection;
 	};
 } // MagicLeap
diff --git a/Engine/Plugins/Lumin/MagicLeap/Source/MagicLeap/Private/AppFramework.cpp b/Engine/Plugins/Lumin/MagicLeap/Source/MagicLeap/Private/AppFramework.cpp
--- a/Engine/Plugins/Lumin/MagicLeap/Source/MagicLeap/Private/AppFramework.cpp
+++ b/Engine/Plugins/Lumin/MagicLeap/Source/MagicLeap/Private/AppFramework.cpp
@@ -38,14 +38,6 @@ void FAppFramework::Startup()
 	base_position_ = FVector::ZeroVector;
 	base_orientation_ = FQuat::Identity;
 
-	{
-		FScopeLock Lock(&EventHandlersCriticalSection);
-	for (auto EventHandler : EventHandlers)
-	{
-		EventHandler->OnAppStartup();
-	}
-	}
-
 	// Register application lifecycle delegates
 	FCoreDelegates::ApplicationWillEnterBackgroundDelegate.AddRaw(this, &FAppFramework::ApplicationPauseDelegate);
 	FCoreDelegates::ApplicationHasEnteredForegroundDelegate.AddRaw(this, &FAppFramework::ApplicationResumeDelegate);
@@ -82,6 +74,12 @@ void FAppFramework::BeginUpdate()
 			base_orientation_ = FQuat::Identity;
 			base_dirty_ = false;
 		}
+
+		FScopeLock Lock(&EventHandlersCriticalSection);
+		for (auto EventHandler : EventHandlers)
+		{
+			EventHandler->OnAppTick();
+		}
 	}
 #endif //WITH_MLSDK
 }
@@ -162,16 +160,16 @@ void FAppFramework::SetBaseRotation(const FRotator& InBaseRotation)
 	base_dirty_ = true;
 }
 
-FTrackingFrame* FAppFramework::GetCurrentFrame() const
+const FTrackingFrame* FAppFramework::GetCurrentFrame() const
 {
 	FMagicLeapHMD* hmd = GEngine ? static_cast<FMagicLeapHMD*>(GEngine->XRSystem->GetHMDDevice()) : nullptr;
-	return hmd ? hmd->GetCurrentFrame() : nullptr;
+	return hmd ? &(hmd->GetCurrentFrame()) : nullptr;
 }
 
-FTrackingFrame* FAppFramework::GetOldFrame() const
+const FTrackingFrame* FAppFramework::GetOldFrame() const
 {
 	FMagicLeapHMD* hmd = GEngine ? static_cast<FMagicLeapHMD*>(GEngine->XRSystem->GetHMDDevice()) : nullptr;
-	return hmd ? hmd->GetOldFrame() : nullptr;
+	return hmd ? &(hmd->GetOldFrame()) : nullptr;
 }
 
 FVector2D FAppFramework::GetFieldOfView() const
@@ -227,14 +225,14 @@ float FAppFramework::GetWorldToMetersScale() const
 
 FTransform FAppFramework::GetCurrentFrameUpdatePose() const
 {
-	FTrackingFrame* frame = GetCurrentFrame();
+	const FTrackingFrame* frame = GetCurrentFrame();
 	return frame ? frame->RawPose : FTransform::Identity;
 }
 
 #if WITH_MLSDK
 bool FAppFramework::GetTransform(const MLCoordinateFrameUID& Id, FTransform& OutTransform, EFailReason& OutReason) const
 {
-	FTrackingFrame* frame = GetCurrentFrame();
+	const FTrackingFrame* frame = GetCurrentFrame();
 	if (frame == nullptr)
 	{
 		OutReason = EFailReason::InvalidTrackingFrame;
@@ -242,7 +240,8 @@ bool FAppFramework::GetTransform(const MLCoordinateFrameUID& Id, FTransform& Out
 	}
 
 	MLTransform transform = MagicLeap::kIdentityTransform;
-	if (MLSnapshotGetTransform(frame->Snapshot, &Id, &transform))
+	MLResult Result = MLSnapshotGetTransform(frame->Snapshot, &Id, &transform);
+	if (Result == MLResult_Ok)
 	{
 		OutTransform = MagicLeap::ToFTransform(transform, GetWorldToMetersScale());
 		if (OutTransform.ContainsNaN())
@@ -261,11 +260,52 @@ bool FAppFramework::GetTransform(const MLCoordinateFrameUID& Id, FTransform& Out
 		OutReason = EFailReason::None;
 		return true;
 	}
+
+#if PLATFORM_LUMIN
+	UE_LOG(LogMagicLeap, Error, TEXT("MLSnapshotGetTransform failed with error %d."), Result);
+#endif
 	OutReason = EFailReason::CallFailed;
 	return false;
 }
 #endif //WITH_MLSDK
 
+TSharedPtr<FCameraCaptureRunnable, ESPMode::ThreadSafe> FAppFramework::GetCameraCaptureRunnable()
+{
+	if (!CameraCaptureRunnable.IsValid())
+	{
+		CameraCaptureRunnable = MakeShared<FCameraCaptureRunnable, ESPMode::ThreadSafe>();
+	}
+	return CameraCaptureRunnable;
+}
+
+void FAppFramework::RefreshCameraCaptureRunnableReferences()
+{
+	// a reference count of 1 is a self reference
+	if (CameraCaptureRunnable.GetSharedReferenceCount() == 1)
+	{
+		CameraCaptureRunnable.Reset();
+	}
+}
+
+TSharedPtr<FImageTrackerRunnable, ESPMode::ThreadSafe> FAppFramework::GetImageTrackerRunnable()
+{
+	if (!ImageTrackerRunnable.IsValid())
+	{
+		ImageTrackerRunnable = MakeShared<FImageTrackerRunnable, ESPMode::ThreadSafe>();
+	}
+
+	return ImageTrackerRunnable;
+}
+
+void FAppFramework::RefreshImageTrackerRunnableReferences()
+{
+	// a reference count of 1 is a self reference
+	if (ImageTrackerRunnable.GetSharedReferenceCount() == 1)
+	{
+		ImageTrackerRunnable.Reset();
+	}
+}
+
 void FAppFramework::AddEventHandler(MagicLeap::IAppEventHandler* EventHandler)
 {
 	FScopeLock Lock(&EventHandlersCriticalSection);
diff --git a/Engine/Plugins/Lumin/MagicLeap/Source/MagicLeap/Private/AppFramework.h b/Engine/Plugins/Lumin/MagicLeap/Source/MagicLeap/Private/AppFramework.h
--- a/Engine/Plugins/Lumin/MagicLeap/Source/MagicLeap/Private/AppFramework.h
+++ b/Engine/Plugins/Lumin/MagicLeap/Source/MagicLeap/Private/AppFramework.h
@@ -25,6 +25,8 @@ ML_INCLUDES_END
 #include "CoreMinimal.h"
 #include "AppEventHandler.h"
 #include "AsyncDestroyer.h"
+#include "CameraCaptureRunnable.h"
+#include "ImageTrackerRunnable.h"
 
 class FMagicLeapHMD;
 struct FTrackingFrame;
@@ -81,13 +83,18 @@ public:
 	bool GetTransform(const MLCoordinateFrameUID& Id, FTransform& OutTransform, EFailReason& OutReason) const;
 #endif //WITH_MLSDK
 
+	TSharedPtr<FCameraCaptureRunnable, ESPMode::ThreadSafe> GetCameraCaptureRunnable();
+	void RefreshCameraCaptureRunnableReferences();
+	TSharedPtr<FImageTrackerRunnable, ESPMode::ThreadSafe> GetImageTrackerRunnable();
+	void RefreshImageTrackerRunnableReferences();
+
 	static void AddEventHandler(MagicLeap::IAppEventHandler* InEventHandler);
 	static void RemoveEventHandler(MagicLeap::IAppEventHandler* InEventHandler);
 	static void AsyncDestroy(MagicLeap::IAppEventHandler* InEventHandler);
 
 private:
-	FTrackingFrame* GetCurrentFrame() const;
-	FTrackingFrame* GetOldFrame() const;
+	const FTrackingFrame* GetCurrentFrame() const;
+	const FTrackingFrame* GetOldFrame() const;
 
 	bool bInitialized = false;
 
@@ -104,6 +111,8 @@ private:
 	static TArray<MagicLeap::IAppEventHandler*> EventHandlers;
 	static FCriticalSection EventHandlersCriticalSection;
 	static MagicLeap::FAsyncDestroyer* AsyncDestroyer;
+	TSharedPtr<FCameraCaptureRunnable, ESPMode::ThreadSafe> CameraCaptureRunnable;
+	TSharedPtr<FImageTrackerRunnable, ESPMode::ThreadSafe> ImageTrackerRunnable;
 };
 
 DEFINE_LOG_CATEGORY_STATIC(LogMagicLeap, Log, All);
\ No newline at end of file
diff --git a/Engine/Plugins/Lumin/MagicLeap/Source/MagicLeap/Private/AsyncDestroyer.cpp b/Engine/Plugins/Lumin/MagicLeap/Source/MagicLeap/Private/AsyncDestroyer.cpp
--- a/Engine/Plugins/Lumin/MagicLeap/Source/MagicLeap/Private/AsyncDestroyer.cpp
+++ b/Engine/Plugins/Lumin/MagicLeap/Source/MagicLeap/Private/AsyncDestroyer.cpp
@@ -38,7 +38,7 @@ namespace MagicLeap
 		while (StopTaskCounter.GetValue() == 0)
 		{
 			IAppEventHandler* EventHandler = nullptr;
-			if (IncomingEventHandlers.Dequeue(EventHandler))
+			while (IncomingEventHandlers.Dequeue(EventHandler))
 			{
 				delete EventHandler;
 				EventHandler = nullptr;
diff --git a/Engine/Plugins/Lumin/MagicLeap/Source/MagicLeap/Private/CameraCaptureComponent.cpp b/Engine/Plugins/Lumin/MagicLeap/Source/MagicLeap/Private/CameraCaptureComponent.cpp
--- a/Engine/Plugins/Lumin/MagicLeap/Source/MagicLeap/Private/CameraCaptureComponent.cpp
+++ b/Engine/Plugins/Lumin/MagicLeap/Source/MagicLeap/Private/CameraCaptureComponent.cpp
@@ -1,21 +1,15 @@
 // Copyright 1998-2018 Epic Games, Inc. All Rights Reserved.
 
 #include "CameraCaptureComponent.h"
+#include "CameraCaptureRunnable.h"
+#include "MagicLeapHMD.h"
+#include "AppFramework.h"
+#include "AppEventHandler.h"
 #include "CoreMinimal.h"
-#include "HAL/Runnable.h"
-#include "HAL/RunnableThread.h"
-#include "HAL/ThreadSafeCounter64.h"
 #include "Engine/Texture2D.h"
 #include "Containers/Queue.h"
 #include "Containers/Array.h"
 #include "Lumin/LuminPlatformFile.h"
-#include "IImageWrapper.h"
-#include "IImageWrapperModule.h"
-#include "Modules/ModuleManager.h"
-#include "IMessageBus.h"
-#include "MessageEndpoint.h"
-#include "MessageEndpointBuilder.h"
-#include "AppEventHandler.h"
 #if PLATFORM_LUMIN
 #include "Lumin/LuminAffinity.h"
 #endif // PLATFORM_LUMIN
@@ -31,401 +25,107 @@ ML_INCLUDES_END
 
 DEFINE_LOG_CATEGORY(LogCameraCapture);
 
-enum CaptureMsgType
-{
-  Request,
-  Response,
-  Log
-};
-
-struct FCaptureMessage
-{
-  CaptureMsgType Type;
-#if WITH_MLSDK
-  MLCameraCaptureType CaptureType;
-#endif //WITH_MLSDK
-  FString Log;
-  FString FilePath;
-  bool Success;
-  UTexture2D* Texture;
-  float Duration;
-
-  FCaptureMessage()
-    : Type(CaptureMsgType::Request)
-#if WITH_MLSDK
-    , CaptureType(MLCameraCaptureType_Image)
-#endif //WITH_MLSDK
-    , Log("")
-    , Success(false)
-    , Texture(nullptr)
-    , Duration(0.0f)
-  {}
-};
-
-class FImageCaptureImpl : public FRunnable, public MagicLeap::IAppEventHandler
+class FCameraCaptureImpl : public MagicLeap::IAppEventHandler
 {
 public:
-  FImageCaptureImpl()
-    : Thread(nullptr)
-    , StopTaskCounter(0)
-    , Semaphore(nullptr)
-		, bCameraConnected(false)
-		, RetryConnectWaitTime(0.5f)
-#if WITH_MLSDK
-    , CameraOutput(nullptr)
-#endif //WITH_MLSDK
-    , ImgExtension(".jpeg")
-    , VidExtension(".mp4")
-  {
-    IImageWrapperModule& ImageWrapperModule = FModuleManager::LoadModuleChecked<IImageWrapperModule>(FName("ImageWrapper"));
-    ImageWrapper = ImageWrapperModule.CreateImageWrapper(EImageFormat::JPEG);	
-    Semaphore = FGenericPlatformProcess::GetSynchEventFromPool(false);
+	FCameraCaptureImpl()
+	: bCapturing(false)
+	{
 #if PLATFORM_LUMIN
-    Thread = FRunnableThread::Create(this, TEXT("FCameraCaptureWorker"), 0, TPri_BelowNormal, FLuminAffinity::GetPoolThreadMask());
-#else
-    Thread = FRunnableThread::Create(this, TEXT("FCameraCaptureWorker"), 0, TPri_BelowNormal);
-#endif // PLATFORM_LUMIN
-  }
-
-  ~FImageCaptureImpl()
-  {
-    StopTaskCounter.Increment();
+		checkf(GEngine, TEXT("[FCameraCaptureImpl::FCameraCaptureImpl()] GEngine is null!"));
+		FAppFramework& AppFramework = static_cast<FMagicLeapHMD*>(GEngine->XRSystem->GetHMDDevice())->GetAppFramework();
+		checkf(AppFramework.IsInitialized(), TEXT("[FCameraCaptureImpl::FCameraCaptureImpl()] AppFramework not yet initialized!"));
+		CameraCaptureRunnable = AppFramework.GetCameraCaptureRunnable();
+#endif //PLATFORM_LUMIN
+	}
 
-	if (Semaphore)
+	virtual ~FCameraCaptureImpl()
 	{
-    Semaphore->Trigger();
-    Thread->WaitForCompletion();
+		// Calling this explicitly to make the chain of destruction more obvious,
+		// ie a potential call to FCameraCaptureRunnable's destructor right here.
+		// On the lumin platform this call will take place inside the destruction
+		// worker thread.
+		CameraCaptureRunnable.Reset();
+		// need this call to prevent the runnable form persisting with a ref count of 1 (due to Appframework's shared pointer instance).
+		if (GEngine && GEngine->XRSystem.IsValid() && GEngine->XRSystem->GetHMDDevice())
+		{
+			static_cast<FMagicLeapHMD*>(GEngine->XRSystem->GetHMDDevice())->GetAppFramework().RefreshCameraCaptureRunnableReferences();
+		}
 	}
 
-    for (auto& CapturedTexture : CapturedTextures)
-    {
-      CapturedTexture->RemoveFromRoot();
-    }
-
-    if (Semaphore)
-    {
-      FGenericPlatformProcess::ReturnSynchEventToPool(Semaphore);
-      Semaphore = nullptr;
-    }
-
-    delete Thread;
-    Thread = nullptr;
-  }
-
-  virtual uint32 Run() override
-  {
-#if WITH_MLSDK
-	  while (StopTaskCounter.GetValue() == 0)
-	  {
-		  if (!bCameraConnected)
-		  {
-			  bCameraConnected = MLCameraConnect();
-			  if (bCameraConnected)
-			  {
-				  MLCameraDeviceStatusCallbacks deviceStatusCallbacks;
-				  FMemory::Memset(&deviceStatusCallbacks, 0, sizeof(MLCameraDeviceStatusCallbacks));
-				  deviceStatusCallbacks.on_preview_buffer_available = OnPreviewBufferAvailable;
-				  if (!MLCameraSetDeviceStatusCallbacks(&deviceStatusCallbacks, nullptr))
-				  {
-					  Log(TEXT("MLCameraSetDeviceStatusCallbacks failed!"));
-				  }
-			  }
-			  else
-			  {
-				  Log(FString::Printf(TEXT("Camera connection attempt failed!  Retrying in %.2f seconds"), RetryConnectWaitTime));
-				  IncomingMessages.Empty();
-				  FPlatformProcess::Sleep(RetryConnectWaitTime);
-			  }
-		  }
-		  else
-		  {
-			  if (IncomingMessages.Dequeue(CurrentMessage))
-			  {
-				  BeginCapture();
-			  }
-
-			  Semaphore->Wait();
-		  }
-	  }
-
-	  if (bCameraConnected)
-	  {
-		  if (!MLCameraDisconnect())
-		  {
-			  UE_LOG(LogCameraCapture, Error, TEXT("MLCameraDisconnect failed!"));
-		  }
-	  }
-
-#endif //WITH_MLSDK
-
-	  return 0;
-  }
-
-  void BeginCapture()
-  {
-#if WITH_MLSDK
-    MLHandle Handle = MLCameraPrepareCapture(CurrentMessage.CaptureType);
-
-    if (Handle == ML_INVALID_HANDLE)
-    {
-      Log(TEXT("MLCameraPrepareCapture failed!  Camera capture aborted!"));
-      EndCapture(false);
-      return;
-    }
-
-    switch (CurrentMessage.CaptureType)
-    {
-    case MLCameraCaptureType_Image:
-    {
-      Log(TEXT("Beginning capture image to file."));
-
-#if PLATFORM_LUMIN
-      IPlatformFile& PlatformFile = IPlatformFile::GetPlatformPhysical();
-      // This module is only for Lumin so this is fine for now.
-      FLuminPlatformFile* LuminPlatformFile = static_cast<FLuminPlatformFile*>(&PlatformFile);
-      UniqueFileName = LuminPlatformFile->ConvertToLuminPath(FPaths::CreateTempFilename(*FPaths::ProjectSavedDir(), TEXT("Img_"), *ImgExtension), true);
-#endif
-      if (!MLCameraCaptureImage(TCHAR_TO_UTF8(*UniqueFileName)))
-      {
-        Log(TEXT("MLCameraCaptureImage failed!  Camera capture aborted!"));
-        EndCapture(false);
-        return;
-      }
-
-      EndCapture(true);
-    }
-    break;
-
-    case MLCameraCaptureType_ImageRaw:
-    {
-      CameraOutput = nullptr;
-
-      Log(TEXT("Beginning capture image to texture."));
-
-      if (!MLCameraCaptureImageRaw())
-      {
-        Log(TEXT("MLCameraCaptureImageRaw failed!  Camera capture aborted!"));
-        EndCapture(false);
-        return;
-      }
-
-      if (!MLCameraGetImageStream(&CameraOutput))
-      {
-        Log(TEXT("MLCameraGetImageStream failed!  Camera capture aborted!"));
-        EndCapture(false);
-        return;
-      }
-
-      if (CameraOutput->plane_count == 0)
-      {
-        Log(TEXT("Invalid plane_count!  Camera capture aborted!"));
-        EndCapture(false);
-        return;
-      }
-
-      EndCapture(true);
-    }
-    break;
-
-    case MLCameraCaptureType_Video:
-    {
-      Log(TEXT("Beginning capture video to file."));
-
-#if PLATFORM_LUMIN
-      IPlatformFile& PlatformFile = IPlatformFile::GetPlatformPhysical();
-      // This module is only for Lumin so this is fine for now.
-      FLuminPlatformFile* LuminPlatformFile = static_cast<FLuminPlatformFile*>(&PlatformFile);
-      UniqueFileName = LuminPlatformFile->ConvertToLuminPath(FPaths::CreateTempFilename(*FPaths::ProjectSavedDir(), TEXT("Vid_"), *VidExtension), true);
-#endif
-      if (!MLCameraCaptureVideoStart(TCHAR_TO_UTF8(*UniqueFileName)))
-      {
-        Log(TEXT("MLCameraCaptureVideoStart failed!  Video capture aborted!"));
-        EndCapture(false);
-        return;
-      }
-
-      FPlatformProcess::Sleep(CurrentMessage.Duration);
-      EndCapture(true);
-    }
-    break;
-  }
-#endif //WITH_MLSDK
-  }
-
-  void EndCapture(bool InSuccess)
-  {
+	bool TryCaptureImageToFile()
+	{
 #if WITH_MLSDK
-    CurrentMessage.Type = CaptureMsgType::Response;
-    CurrentMessage.Success = InSuccess;
-	CurrentMessage.FilePath = UniqueFileName;
-
-    if (!InSuccess)
-    {
-      uint32 DeviceStatus = 0;
-      if (MLCameraGetDeviceStatus(&DeviceStatus))
-      {
-        if ((MLCameraDeviceStatusFlag_Available & DeviceStatus) != 0)
-        {
-          Log(TEXT("Device status = MLCameraDeviceStatusFlag_Available"));
-        }
-        else if ((MLCameraDeviceStatusFlag_Opened & DeviceStatus) != 0)
-        {
-          Log(TEXT("Device status = MLCameraDeviceStatusFlag_Opened"));
-        }
-        else if ((MLCameraDeviceStatusFlag_Disconnected  & DeviceStatus) != 0)
-        {
-          Log(TEXT("Device status = MLCameraDeviceStatusFlag_Disconnected "));
-        }
-        else if ((MLCameraDeviceStatusFlag_Error & DeviceStatus) != 0)
-        {
-          Log(TEXT("Device status = MLCameraDeviceStatusFlag_Error"));
-          MLCameraError CameraError;
-          if (MLCameraGetErrorCode(&CameraError))
-          {
-            switch (CameraError)
-            {
-            case MLCameraError_None: Log(TEXT("Error = MLCameraError_None")); break;
-            case MLCameraError_Invalid: Log(TEXT("Error = MLCameraError_Invalid")); break;
-            case MLCameraError_Disabled: Log(TEXT("Error = MLCameraError_Disabled")); break;
-            case MLCameraError_DeviceFailed: Log(TEXT("Error = MLCameraError_DeviceFailed")); break;
-            case MLCameraError_ServiceFailed: Log(TEXT("Error = MLCameraError_ServiceFailed")); break;
-            case MLCameraError_CaptureFailed: Log(TEXT("Error = MLCameraError_CaptureFailed")); break;
-            case MLCameraError_Unknown: Log(TEXT("Error = MLCameraError_Unknown")); break;
-            }
-          }
-          else
-          {
-            Log(TEXT("MLCameraGetErrorCode failed!"));
-          }
-        }
-      }
-      else
-      {
-        Log(TEXT("MLCameraGetDeviceStatus failed!"));
-      }
-    }
-
-    switch (CurrentMessage.CaptureType)
-    {
-    case MLCameraCaptureType_Image:
-    {
-      if (InSuccess)
-      {
-        Log(FString::Printf(TEXT("Captured image to %s"), *UniqueFileName));
-      }
-    }
-    break;
-
-    case MLCameraCaptureType_ImageRaw:
-    {
-      if (InSuccess)
-      {
-        MLCameraPlaneInfo& ImageInfo = CameraOutput->planes[0];
-        //JMC: ImageInfo.width/height are incorrect, ImageInfo.data contains header info also as at MLSKD 0.9.0
-        // using IImageWrapperModule as workaround/potentially permanent solution
-        if (ImageWrapper.IsValid() && ImageWrapper->SetCompressed(ImageInfo.data, ImageInfo.size))
-        {
-          const TArray<uint8>* RawData = NULL;
-          if (ImageWrapper->GetRaw(ImageWrapper->GetFormat(), 8, RawData))
-          {
-            Log(FString::Printf(TEXT("ImageWrapper width=%d height=%d size=%d"), ImageWrapper->GetWidth(), ImageWrapper->GetHeight(), RawData->Num()));
-            UTexture2D* CaptureTexture = UTexture2D::CreateTransient(ImageWrapper->GetWidth(), ImageWrapper->GetHeight(), EPixelFormat::PF_B8G8R8A8);
-            CaptureTexture->AddToRoot();
-            CapturedTextures.Add(CaptureTexture);
-            FTexture2DMipMap& Mip = CaptureTexture->PlatformData->Mips[0];
-            void* Data = Mip.BulkData.Lock(LOCK_READ_WRITE);
-            FMemory::Memcpy(Data, RawData->GetData(), Mip.BulkData.GetBulkDataSize());
-            Mip.BulkData.Unlock();
-            CaptureTexture->UpdateResource();
-            CurrentMessage.Texture = CaptureTexture;
-          }
-        }
-      }     
-    }
-    break;
-
-    case MLCameraCaptureType_Video:
-    {
-      if (InSuccess)
-      {
-        MLCameraCaptureVideoStop();
-        Log(FString::Printf(TEXT("Captured video to %s"), *UniqueFileName));
-      }
-    }
-    break;
-    }
-
-	// clear any left over events from this capture session
-	IncomingMessages.Empty();
-	// signal the update thread to handle the capture result
-    OutgoingMessages.Enqueue(CurrentMessage);
+		if (!bCapturing)
+		{
+			bCapturing = true;
+			FCaptureMessage Msg;
+			Msg.Type = CaptureMsgType::Request;
+			Msg.CaptureType = MLCameraCaptureType_Image;
+			Msg.Requester = this;
+			CameraCaptureRunnable->ProcessCaptureMessage(Msg);
+			return true;
+		}
 #endif //WITH_MLSDK
-  }
-
-  void Log(const FString& Info)
-  {
-    FCaptureMessage Msg;
-    Msg.Type = CaptureMsgType::Log;
-    Msg.Log = Info;
-    OutgoingMessages.Enqueue(Msg);
-  }
-
-  void ProcessCaptureMessage(const FCaptureMessage& InMsg)
-  {
-	  if (bCameraConnected)
-	  {
-		  IncomingMessages.Enqueue(InMsg);
-		  // wake up the worker to process the event
-		  Semaphore->Trigger();
-	  }
-	  else
-	  {
-		  Log(TEXT("Discarding capture task due to lack of camera connection."));
-	  }
-  }
-
-  /** Internal thread this instance is running on */
-  FRunnableThread* Thread;
-  FThreadSafeCounter StopTaskCounter;
-  TQueue<FCaptureMessage, EQueueMode::Spsc> IncomingMessages;
-  TQueue<FCaptureMessage, EQueueMode::Spsc> OutgoingMessages;
-  FEvent* Semaphore;
-  FCaptureMessage CurrentMessage;
-  bool bCameraConnected;
-  const float RetryConnectWaitTime;
+		return false;
+	}
 
+	bool TryCaptureImageToTexture()
+	{
 #if WITH_MLSDK
-  MLCameraDeviceStatusCallbacks DeviceStatusCallbacks;
-  MLCameraCaptureCallbacks CaptureCallbacks;
-  MLCameraOutput* CameraOutput;
+		if (!bCapturing)
+		{
+			bCapturing = true;
+			FCaptureMessage Msg;
+			Msg.Type = CaptureMsgType::Request;
+			Msg.CaptureType = MLCameraCaptureType_ImageRaw;
+			Msg.Requester = this;
+			CameraCaptureRunnable->ProcessCaptureMessage(Msg);
+			return true;
+		}
 #endif //WITH_MLSDK
+		return false;
+	}
 
-  const FString ImgExtension;
-  const FString VidExtension;
-  FString UniqueFileName;
-  TArray<UTexture2D*> CapturedTextures;
-  TSharedPtr<IImageWrapper> ImageWrapper;
+	bool TryCaptureVideoToFile(float InDuration)
+	{
 #if WITH_MLSDK
-  static FThreadSafeCounter64 PreviewHandle;
+		if (!bCapturing)
+		{
+			bCapturing = true;
+			FCaptureMessage Msg;
+			Msg.Type = CaptureMsgType::Request;
+			Msg.CaptureType = MLCameraCaptureType_Video;
+			Msg.Duration = InDuration;
+			Msg.Requester = this;
+			CameraCaptureRunnable->ProcessCaptureMessage(Msg);
+			return true;
+		}
 #endif //WITH_MLSDK
+		return false;
+	}
 
-private:
-#if WITH_MLSDK
-	static void OnPreviewBufferAvailable(MLHandle Output, void *Data)
+	bool TryGetResult(FCaptureMessage& OutMsg)
 	{
-		(void)Data;
-		FImageCaptureImpl::PreviewHandle.Set(static_cast<int64>(Output));
+		if (bCapturing && CameraCaptureRunnable->OutgoingMessages.Peek(OutMsg))
+		{
+			if (OutMsg.Requester == this)
+			{
+				CameraCaptureRunnable->OutgoingMessages.Pop();
+				return true;
+			}
+		}
+		
+		return false;
 	}
-#endif //WITH_MLSDK
-};
 
-#if WITH_MLSDK
-FThreadSafeCounter64 FImageCaptureImpl::PreviewHandle = ML_INVALID_HANDLE;
-#endif //WITH_MLSDK
+public:
+	bool bCapturing;
+	TSharedPtr<FCameraCaptureRunnable, ESPMode::ThreadSafe> CameraCaptureRunnable;
+};
 
 UCameraCaptureComponent::UCameraCaptureComponent()
     : Impl(nullptr)
-    , bCapturing(false)
 {
   PrimaryComponentTick.TickGroup = TG_PrePhysics;
   PrimaryComponentTick.bStartWithTickEnabled = true;
@@ -441,15 +141,15 @@ UCameraCaptureComponent::~UCameraCaptureComponent()
 		Impl = nullptr;
 	}
 #else
-  delete Impl;
-  Impl = nullptr;
+	delete Impl;
+	Impl = nullptr;
 #endif // PLATFORM_LUMIN	
 }
 
 void UCameraCaptureComponent::BeginPlay()
 {
   Super::BeginPlay();
-  Impl = new FImageCaptureImpl;
+  Impl = new FCameraCaptureImpl;
 }
 
 void UCameraCaptureComponent::TickComponent(float DeltaTime, enum ELevelTick TickType, FActorComponentTickFunction *ThisTickFunction)
@@ -457,11 +157,9 @@ void UCameraCaptureComponent::TickComponent(float DeltaTime, enum ELevelTick Tic
 #if WITH_MLSDK
   Super::TickComponent(DeltaTime, TickType, ThisTickFunction);
 
-  if (!Impl->OutgoingMessages.IsEmpty())
+  FCaptureMessage Msg;
+  if (Impl->TryGetResult(Msg))
   {
-    FCaptureMessage Msg;
-    Impl->OutgoingMessages.Dequeue(Msg);
-
     if (Msg.Type == CaptureMsgType::Request)
     {
       Log("Unexpected CaptureMsgType::Request received from worker thread!");
@@ -514,7 +212,7 @@ void UCameraCaptureComponent::TickComponent(float DeltaTime, enum ELevelTick Tic
       break;
       }
 
-      bCapturing = false;
+      Impl->bCapturing = false;
     }
   }
 #endif //WITH_MLSDK
@@ -557,15 +255,8 @@ UCameraCaptureComponent::FCameraCaptureVidToFileFailure& UCameraCaptureComponent
 
 bool UCameraCaptureComponent::CaptureImageToFileAsync()
 {
-  if (!bCapturing)
+  if (Impl->TryCaptureImageToFile())
   {
-    bCapturing = true;
-    FCaptureMessage Msg;
-    Msg.Type = CaptureMsgType::Request;
-#if WITH_MLSDK
-    Msg.CaptureType = MLCameraCaptureType_Image;
-#endif //WITH_MLSDK
-    Impl->ProcessCaptureMessage(Msg);
     return true;
   }
 
@@ -575,15 +266,8 @@ bool UCameraCaptureComponent::CaptureImageToFileAsync()
 
 bool UCameraCaptureComponent::CaptureImageToTextureAsync()
 {
-  if (!bCapturing)
+  if (Impl->TryCaptureImageToTexture())
   {
-    bCapturing = true;
-    FCaptureMessage Msg;
-    Msg.Type = CaptureMsgType::Request;
-#if WITH_MLSDK
-    Msg.CaptureType = MLCameraCaptureType_ImageRaw;
-#endif //WITH_MLSDK
-    Impl->ProcessCaptureMessage(Msg);
     return true;
   }
 
@@ -593,16 +277,8 @@ bool UCameraCaptureComponent::CaptureImageToTextureAsync()
 
 bool UCameraCaptureComponent::CaptureVideoToFileAsync(float InDuration)
 {
-  if (!bCapturing)
+  if (Impl->TryCaptureVideoToFile(InDuration))
   {
-    bCapturing = true;
-    FCaptureMessage Msg;
-    Msg.Type = CaptureMsgType::Request;
-#if WITH_MLSDK
-    Msg.CaptureType = MLCameraCaptureType_Video;
-#endif //WITH_MLSDK
-    Msg.Duration = InDuration;
-    Impl->ProcessCaptureMessage(Msg);
     return true;
   }
 
@@ -613,7 +289,7 @@ bool UCameraCaptureComponent::CaptureVideoToFileAsync(float InDuration)
 int64 UCameraCaptureComponent::GetPreviewHandle()
 {
 #if WITH_MLSDK
-	return FImageCaptureImpl::PreviewHandle.GetValue();
+	return FCameraCaptureRunnable::PreviewHandle.GetValue();
 #else
 	return -1;
 #endif //WITH_MLSDK
diff --git a/Engine/Plugins/Lumin/MagicLeap/Source/MagicLeap/Private/CameraCaptureRunnable.cpp b/Engine/Plugins/Lumin/MagicLeap/Source/MagicLeap/Private/CameraCaptureRunnable.cpp
new file mode 100644
--- /dev/null
+++ b/Engine/Plugins/Lumin/MagicLeap/Source/MagicLeap/Private/CameraCaptureRunnable.cpp
@@ -0,0 +1,420 @@
+// %BANNER_BEGIN%
+// ---------------------------------------------------------------------
+// %COPYRIGHT_BEGIN%
+//
+// Copyright (c) 2017 Magic Leap, Inc. (COMPANY) All Rights Reserved.
+// Magic Leap, Inc. Confidential and Proprietary
+//
+// NOTICE: All information contained herein is, and remains the property
+// of COMPANY. The intellectual and technical concepts contained herein
+// are proprietary to COMPANY and may be covered by U.S. and Foreign
+// Patents, patents in process, and are protected by trade secret or
+// copyright law. Dissemination of this information or reproduction of
+// this material is strictly forbidden unless prior written permission is
+// obtained from COMPANY. Access to the source code contained herein is
+// hereby forbidden to anyone except current COMPANY employees, managers
+// or contractors who have executed Confidentiality and Non-disclosure
+// agreements explicitly covering such access.
+//
+// The copyright notice above does not evidence any actual or intended
+// publication or disclosure of this source code, which includes
+// information that is confidential and/or proprietary, and is a trade
+// secret, of COMPANY. ANY REPRODUCTION, MODIFICATION, DISTRIBUTION,
+// PUBLIC PERFORMANCE, OR PUBLIC DISPLAY OF OR THROUGH USE OF THIS
+// SOURCE CODE WITHOUT THE EXPRESS WRITTEN CONSENT OF COMPANY IS
+// STRICTLY PROHIBITED, AND IN VIOLATION OF APPLICABLE LAWS AND
+// INTERNATIONAL TREATIES. THE RECEIPT OR POSSESSION OF THIS SOURCE
+// CODE AND/OR RELATED INFORMATION DOES NOT CONVEY OR IMPLY ANY RIGHTS
+// TO REPRODUCE, DISCLOSE OR DISTRIBUTE ITS CONTENTS, OR TO MANUFACTURE,
+// USE, OR SELL ANYTHING THAT IT MAY DESCRIBE, IN WHOLE OR IN PART.
+//
+// %COPYRIGHT_END%
+// --------------------------------------------------------------------
+// %BANNER_END%
+
+#include "CameraCaptureRunnable.h"
+#include "Engine/Engine.h"
+#include "CoreMinimal.h"
+#include "Modules/ModuleManager.h"
+#include "Engine/Texture2D.h"
+#include "AppFramework.h"
+#include "Lumin/LuminPlatformFile.h"
+#if PLATFORM_LUMIN
+#include "Lumin/LuminAffinity.h"
+#endif // PLATFORM_LUMIN
+
+DEFINE_LOG_CATEGORY(LogCameraCaptureRunnable);
+
+#if WITH_MLSDK
+FThreadSafeCounter64 FCameraCaptureRunnable::PreviewHandle = ML_INVALID_HANDLE;
+#endif //WITH_MLSDK
+
+FCameraCaptureRunnable::FCameraCaptureRunnable()
+#if WITH_MLSDK
+	: MagicLeap::IAppEventHandler({ MLPrivilegeID_CameraCapture, MLPrivilegeID_AudioRecognizer, MLPrivilegeID_AudioCaptureMic, MLPrivilegeID_VoiceInput })
+#else
+	: MagicLeap::IAppEventHandler()
+#endif //WITH_MLSDK
+	, Thread(nullptr)
+	, StopTaskCounter(0)
+	, Semaphore(nullptr)
+#if WITH_MLSDK
+	, CameraOutput(nullptr)
+#endif // WITH_MLSDK
+	, bCameraConnected(false)
+	, RetryConnectWaitTime(0.5f)
+	, ImgExtension(".jpeg")
+	, VidExtension(".mp4")
+	, bPaused(false)
+{
+	IImageWrapperModule& ImageWrapperModule = FModuleManager::LoadModuleChecked<IImageWrapperModule>(FName("ImageWrapper"));
+	ImageWrapper = ImageWrapperModule.CreateImageWrapper(EImageFormat::JPEG);
+	Semaphore = FGenericPlatformProcess::GetSynchEventFromPool(false);
+#if PLATFORM_LUMIN
+	Thread = FRunnableThread::Create(this, TEXT("FCameraCaptureWorker"), 0, TPri_BelowNormal, FLuminAffinity::GetPoolThreadMask());
+#else
+	Thread = FRunnableThread::Create(this, TEXT("FCameraCaptureWorker"), 0, TPri_BelowNormal);
+#endif // PLATFORM_LUMIN
+}
+
+FCameraCaptureRunnable::~FCameraCaptureRunnable()
+{
+	Stop();
+}
+
+uint32 FCameraCaptureRunnable::Run()
+{
+#if WITH_MLSDK
+	while (StopTaskCounter.GetValue() == 0)
+	{
+		if (bPaused)
+		{
+			if (bCameraConnected)
+			{
+				MLResult Result = MLCameraDisconnect();
+				UE_CLOG(Result != MLResult_Ok, LogCameraCaptureRunnable, Error, TEXT("MLCameraDisconnect failed with error %s!"), UTF8_TO_TCHAR(MLGetResultString(Result)));
+				bCameraConnected = false;
+			}
+
+			Semaphore->Wait();
+		}
+
+		if (IncomingMessages.Dequeue(CurrentMessage))
+		{
+			while (!bCameraConnected)
+			{
+				if (GetPrivilegeStatus(MLPrivilegeID_CameraCapture) == MagicLeap::EPrivilegeState::Granted)
+				{
+					MLResult Result = MLCameraConnect();
+					if (Result == MLResult_Ok)
+					{
+						bCameraConnected = true;
+						MLCameraDeviceStatusCallbacks deviceStatusCallbacks;
+						FMemory::Memset(&deviceStatusCallbacks, 0, sizeof(MLCameraDeviceStatusCallbacks));
+						deviceStatusCallbacks.on_preview_buffer_available = OnPreviewBufferAvailable;
+						Result = MLCameraSetDeviceStatusCallbacks(&deviceStatusCallbacks, nullptr);
+						if (Result != MLResult_Ok)
+						{
+							Log(FString::Printf(TEXT("MLCameraSetDeviceStatusCallbacks failed with error %s!"), UTF8_TO_TCHAR(MLGetResultString(Result))));
+						}
+					}
+					else
+					{
+						bCameraConnected = false;
+						Log(FString::Printf(TEXT("MLCameraConnect failed with error %s!  Retrying in %.2f seconds"), UTF8_TO_TCHAR(MLGetResultString(Result)), RetryConnectWaitTime));
+						IncomingMessages.Empty();
+						FPlatformProcess::Sleep(RetryConnectWaitTime);
+					}
+				}
+				else
+				{
+					Log(TEXT("Cannot connect to camera due to lack of privilege!"));
+					break;
+				}
+			}
+
+			BeginCapture();
+			Semaphore->Wait();
+		}
+	}
+
+	if (bCameraConnected)
+	{
+		MLResult Result = MLCameraDisconnect();
+		UE_CLOG(Result != MLResult_Ok, LogCameraCaptureRunnable, Error, TEXT("MLCameraDisconnect failed with error %s!"), UTF8_TO_TCHAR(MLGetResultString(Result)));
+	}
+#endif //WITH_MLSDK
+	return 0;
+}
+
+void FCameraCaptureRunnable::Stop()
+{
+	StopTaskCounter.Increment();
+
+	if (Semaphore)
+	{
+		Semaphore->Trigger();
+		Thread->WaitForCompletion();
+		FGenericPlatformProcess::ReturnSynchEventToPool(Semaphore);
+		Semaphore = nullptr;
+	}
+
+	delete Thread;
+	Thread = nullptr;
+}
+
+void FCameraCaptureRunnable::BeginCapture()
+{
+#if WITH_MLSDK
+	MLHandle Handle = ML_INVALID_HANDLE;
+	MLResult Result = MLCameraPrepareCapture(CurrentMessage.CaptureType, &Handle);
+	if (Result != MLResult_Ok)
+	{
+		Log(FString::Printf(TEXT("MLCameraPrepareCapture failed with error %s!  Camera capture aborted!"), UTF8_TO_TCHAR(MLGetResultString(Result))));
+		EndCapture(false);
+		return;
+	}
+
+	switch (CurrentMessage.CaptureType)
+	{
+	case MLCameraCaptureType_Image:
+	{
+		Log(TEXT("Beginning capture image to file."));
+
+#if PLATFORM_LUMIN
+		IPlatformFile& PlatformFile = IPlatformFile::GetPlatformPhysical();
+		// This module is only for Lumin so this is fine for now.
+		FLuminPlatformFile* LuminPlatformFile = static_cast<FLuminPlatformFile*>(&PlatformFile);
+		UniqueFileName = LuminPlatformFile->ConvertToLuminPath(FPaths::CreateTempFilename(*FPaths::ProjectSavedDir(), TEXT("Img_"), *ImgExtension), true);
+#endif
+		Result = MLCameraCaptureImage(TCHAR_TO_UTF8(*UniqueFileName));
+		if (Result != MLResult_Ok)
+		{
+			Log(FString::Printf(TEXT("MLCameraCaptureImage failed with error %s!  Camera capture aborted!"), UTF8_TO_TCHAR(MLGetResultString(Result))));
+			EndCapture(false);
+			return;
+		}
+
+		EndCapture(true);
+	}
+	break;
+
+	case MLCameraCaptureType_ImageRaw:
+	{
+		CameraOutput = nullptr;
+
+		Log(TEXT("Beginning capture image to texture."));
+		Result = MLCameraCaptureImageRaw();
+		if (Result != MLResult_Ok)
+		{
+			Log(FString::Printf(TEXT("MLCameraCaptureImageRaw failed with error %s!  Camera capture aborted!"), UTF8_TO_TCHAR(MLGetResultString(Result))));
+			EndCapture(false);
+			return;
+		}
+
+		Result = MLCameraGetImageStream(&CameraOutput);
+		if (Result != MLResult_Ok)
+		{
+			Log(FString::Printf(TEXT("MLCameraGetImageStream failed with error %s!  Camera capture aborted!"), UTF8_TO_TCHAR(MLGetResultString(Result))));
+			EndCapture(false);
+			return;
+		}
+
+		if (CameraOutput->plane_count == 0)
+		{
+			Log(TEXT("Invalid plane_count!  Camera capture aborted!"));
+			EndCapture(false);
+			return;
+		}
+
+		EndCapture(true);
+	}
+	break;
+
+	case MLCameraCaptureType_Video:
+	{
+		Log(TEXT("Beginning capture video to file."));
+		if (GetPrivilegeStatus(MLPrivilegeID_AudioRecognizer) == MagicLeap::EPrivilegeState::Granted &&
+			GetPrivilegeStatus(MLPrivilegeID_AudioCaptureMic) == MagicLeap::EPrivilegeState::Granted &&
+			GetPrivilegeStatus(MLPrivilegeID_VoiceInput) == MagicLeap::EPrivilegeState::Granted)
+		{
+#if PLATFORM_LUMIN
+			IPlatformFile& PlatformFile = IPlatformFile::GetPlatformPhysical();
+			// This module is only for Lumin so this is fine for now.
+			FLuminPlatformFile* LuminPlatformFile = static_cast<FLuminPlatformFile*>(&PlatformFile);
+			UniqueFileName = LuminPlatformFile->ConvertToLuminPath(FPaths::CreateTempFilename(*FPaths::ProjectSavedDir(), TEXT("Vid_"), *VidExtension), true);
+#endif
+			Result = MLCameraCaptureVideoStart(TCHAR_TO_UTF8(*UniqueFileName));
+			if (Result != MLResult_Ok)
+			{
+				Log(FString::Printf(TEXT("MLCameraCaptureVideoStart failed with error %s!  Video capture aborted!"), UTF8_TO_TCHAR(MLGetResultString(Result))));
+				EndCapture(false);
+				return;
+			}
+
+			FPlatformProcess::Sleep(CurrentMessage.Duration);
+			EndCapture(true);
+		}
+		else
+		{
+			Log(TEXT("Cannot capture video due to lack of privilege!"));
+			EndCapture(false);
+		}
+
+	}
+	break;
+	}
+#endif //WITH_MLSDK
+}
+
+void FCameraCaptureRunnable::EndCapture(bool InSuccess)
+{
+#if WITH_MLSDK
+	CurrentMessage.Type = CaptureMsgType::Response;
+	CurrentMessage.Success = InSuccess;
+	CurrentMessage.FilePath = UniqueFileName;
+
+	if (!InSuccess)
+	{
+		uint32 DeviceStatus = 0;
+		MLResult Result = MLCameraGetDeviceStatus(&DeviceStatus);
+		if (Result == MLResult_Ok)
+		{
+			if ((MLCameraDeviceStatusFlag_Available & DeviceStatus) != 0)
+			{
+				Log(TEXT("Device status = MLCameraDeviceStatusFlag_Available"));
+			}
+			else if ((MLCameraDeviceStatusFlag_Opened & DeviceStatus) != 0)
+			{
+				Log(TEXT("Device status = MLCameraDeviceStatusFlag_Opened"));
+			}
+			else if ((MLCameraDeviceStatusFlag_Disconnected  & DeviceStatus) != 0)
+			{
+				Log(TEXT("Device status = MLCameraDeviceStatusFlag_Disconnected "));
+			}
+			else if ((MLCameraDeviceStatusFlag_Error & DeviceStatus) != 0)
+			{
+				Log(TEXT("Device status = MLCameraDeviceStatusFlag_Error"));
+				MLCameraError CameraError;
+				if (MLCameraGetErrorCode(&CameraError) == MLResult_Ok)
+				{
+					switch (CameraError)
+					{
+					case MLCameraError_None: Log(TEXT("Error = MLCameraError_None")); break;
+					case MLCameraError_Invalid: Log(TEXT("Error = MLCameraError_Invalid")); break;
+					case MLCameraError_Disabled: Log(TEXT("Error = MLCameraError_Disabled")); break;
+					case MLCameraError_DeviceFailed: Log(TEXT("Error = MLCameraError_DeviceFailed")); break;
+					case MLCameraError_ServiceFailed: Log(TEXT("Error = MLCameraError_ServiceFailed")); break;
+					case MLCameraError_CaptureFailed: Log(TEXT("Error = MLCameraError_CaptureFailed")); break;
+					case MLCameraError_Unknown: Log(TEXT("Error = MLCameraError_Unknown")); break;
+					}
+				}
+				else
+				{
+					Log(TEXT("MLCameraGetErrorCode failed!"));
+				}
+			}
+		}
+		else
+		{
+			Log(FString::Printf(TEXT("MLCameraGetDeviceStatus failed with error %s!"), UTF8_TO_TCHAR(MLGetResultString(Result))));
+		}
+	}
+
+	if (InSuccess)
+	{
+		switch (CurrentMessage.CaptureType)
+		{
+		case MLCameraCaptureType_Image:
+		{
+			Log(FString::Printf(TEXT("Captured image to %s"), *UniqueFileName));
+		}
+		break;
+
+		case MLCameraCaptureType_ImageRaw:
+		{
+			MLCameraPlaneInfo& ImageInfo = CameraOutput->planes[0];
+			if (ImageWrapper.IsValid() && ImageWrapper->SetCompressed(ImageInfo.data, ImageInfo.size))
+			{
+				const TArray<uint8>* RawData = NULL;
+				if (ImageWrapper->GetRaw(ImageWrapper->GetFormat(), 8, RawData))
+				{
+					Log(FString::Printf(TEXT("ImageWrapper width=%d height=%d size=%d"), ImageWrapper->GetWidth(), ImageWrapper->GetHeight(), RawData->Num()));
+					UTexture2D* CaptureTexture = UTexture2D::CreateTransient(ImageWrapper->GetWidth(), ImageWrapper->GetHeight(), EPixelFormat::PF_B8G8R8A8);
+					CaptureTexture->AddToRoot();
+					FTexture2DMipMap& Mip = CaptureTexture->PlatformData->Mips[0];
+					void* Data = Mip.BulkData.Lock(LOCK_READ_WRITE);
+					FMemory::Memcpy(Data, RawData->GetData(), Mip.BulkData.GetBulkDataSize());
+					Mip.BulkData.Unlock();
+					CaptureTexture->UpdateResource();
+					CurrentMessage.Texture = CaptureTexture;
+				}
+			}
+		}
+		break;
+
+		case MLCameraCaptureType_Video:
+		{
+			MLResult Result = MLCameraCaptureVideoStop();
+			if (Result != MLResult_Ok)
+			{
+				Log(FString::Printf(TEXT("MLCameraCaptureVideoStop failed with error %s!"), UTF8_TO_TCHAR(MLGetResultString(Result))));
+			}
+			else
+			{
+				Log(FString::Printf(TEXT("Captured video to %s"), *UniqueFileName));
+			}
+		}
+		break;
+		}
+	}
+
+	// clear any left over events from this capture session
+	IncomingMessages.Empty();
+	// signal the update thread to handle the capture result
+	OutgoingMessages.Enqueue(CurrentMessage);
+#endif //WITH_MLSDK
+}
+
+void FCameraCaptureRunnable::Log(const FString& Info)
+{
+	FCaptureMessage Msg;
+	Msg.Type = CaptureMsgType::Log;
+	Msg.Log = Info;
+	Msg.Requester = CurrentMessage.Requester;
+	OutgoingMessages.Enqueue(Msg);
+}
+
+void FCameraCaptureRunnable::ProcessCaptureMessage(const FCaptureMessage& InMsg)
+{
+	IncomingMessages.Enqueue(InMsg);
+	// wake up the worker to process the event
+	Semaphore->Trigger();
+}
+
+void FCameraCaptureRunnable::OnAppPause()
+{
+	IAppEventHandler::OnAppPause();
+	bPaused = true;
+	Semaphore->Trigger();
+}
+
+void FCameraCaptureRunnable::OnAppResume()
+{
+	IAppEventHandler::OnAppResume();
+	bPaused = false;
+	Semaphore->Trigger();
+}
+
+void FCameraCaptureRunnable::OnAppShutDown()
+{
+	Stop();
+}
+
+#if WITH_MLSDK
+void FCameraCaptureRunnable::OnPreviewBufferAvailable(MLHandle Output, void *Data)
+{
+	(void)Data;
+	FCameraCaptureRunnable::PreviewHandle.Set(static_cast<int64>(Output));
+}
+#endif //WITH_MLSDK
\ No newline at end of file
diff --git a/Engine/Plugins/Lumin/MagicLeap/Source/MagicLeap/Private/CameraCaptureRunnable.h b/Engine/Plugins/Lumin/MagicLeap/Source/MagicLeap/Private/CameraCaptureRunnable.h
new file mode 100644
--- /dev/null
+++ b/Engine/Plugins/Lumin/MagicLeap/Source/MagicLeap/Private/CameraCaptureRunnable.h
@@ -0,0 +1,142 @@
+// %BANNER_BEGIN%
+// ---------------------------------------------------------------------
+// %COPYRIGHT_BEGIN%
+//
+// Copyright (c) 2017 Magic Leap, Inc. (COMPANY) All Rights Reserved.
+// Magic Leap, Inc. Confidential and Proprietary
+//
+// NOTICE: All information contained herein is, and remains the property
+// of COMPANY. The intellectual and technical concepts contained herein
+// are proprietary to COMPANY and may be covered by U.S. and Foreign
+// Patents, patents in process, and are protected by trade secret or
+// copyright law. Dissemination of this information or reproduction of
+// this material is strictly forbidden unless prior written permission is
+// obtained from COMPANY. Access to the source code contained herein is
+// hereby forbidden to anyone except current COMPANY employees, managers
+// or contractors who have executed Confidentiality and Non-disclosure
+// agreements explicitly covering such access.
+//
+// The copyright notice above does not evidence any actual or intended
+// publication or disclosure of this source code, which includes
+// information that is confidential and/or proprietary, and is a trade
+// secret, of COMPANY. ANY REPRODUCTION, MODIFICATION, DISTRIBUTION,
+// PUBLIC PERFORMANCE, OR PUBLIC DISPLAY OF OR THROUGH USE OF THIS
+// SOURCE CODE WITHOUT THE EXPRESS WRITTEN CONSENT OF COMPANY IS
+// STRICTLY PROHIBITED, AND IN VIOLATION OF APPLICABLE LAWS AND
+// INTERNATIONAL TREATIES. THE RECEIPT OR POSSESSION OF THIS SOURCE
+// CODE AND/OR RELATED INFORMATION DOES NOT CONVEY OR IMPLY ANY RIGHTS
+// TO REPRODUCE, DISCLOSE OR DISTRIBUTE ITS CONTENTS, OR TO MANUFACTURE,
+// USE, OR SELL ANYTHING THAT IT MAY DESCRIBE, IN WHOLE OR IN PART.
+//
+// %COPYRIGHT_END%
+// --------------------------------------------------------------------
+// %BANNER_END%
+
+#pragma once
+
+#include "HAL/Runnable.h"
+#include "HAL/RunnableThread.h"
+#include "HAL/ThreadSafeBool.h"
+#include "HAL/ThreadSafeCounter64.h"
+#include "Containers/Queue.h"
+#include "Containers/Array.h"
+#include "IImageWrapper.h"
+#include "IImageWrapperModule.h"
+#include "AppEventHandler.h"
+#include "MagicLeapPluginUtil.h" // for ML_INCLUDES_START/END
+#if WITH_MLSDK
+ML_INCLUDES_START
+#include <ml_api.h>
+#include <ml_camera.h>
+#include <ml_camera_metadata.h>
+#include <ml_camera_metadata_tags.h>
+ML_INCLUDES_END
+#endif //WITH_MLSDK
+
+enum class CaptureMsgType : uint32
+{
+	Request,
+	Response,
+	Log
+};
+
+struct FCaptureMessage
+{
+	CaptureMsgType Type;
+#if WITH_MLSDK
+	MLCameraCaptureType CaptureType;
+#endif //WITH_MLSDK
+	FString Log;
+	FString FilePath;
+	bool Success;
+	class UTexture2D* Texture;
+	float Duration;
+	class FCameraCaptureImpl* Requester;
+
+	FCaptureMessage()
+		: Type(CaptureMsgType::Request)
+#if WITH_MLSDK
+		, CaptureType(MLCameraCaptureType_Image)
+#endif //WITH_MLSDK
+		, Log("")
+		, Success(false)
+		, Texture(nullptr)
+		, Duration(0.0f)
+		, Requester(nullptr)
+	{}
+};
+
+class FCameraCaptureRunnable : public FRunnable, public MagicLeap::IAppEventHandler
+{
+public:
+	FCameraCaptureRunnable();
+	virtual ~FCameraCaptureRunnable();
+
+	uint32 Run() override;
+
+	void Stop() override;
+
+	void BeginCapture();
+
+	void EndCapture(bool InSuccess);
+
+	void Log(const FString& Info);
+
+	void ProcessCaptureMessage(const FCaptureMessage& InMsg);
+
+	void OnAppPause() override;
+
+	void OnAppResume() override;
+
+	void OnAppShutDown() override;
+
+	/** Internal thread this instance is running on */
+	FRunnableThread* Thread;
+	FThreadSafeCounter StopTaskCounter;
+	TQueue<FCaptureMessage, EQueueMode::Spsc> IncomingMessages;
+	TQueue<FCaptureMessage, EQueueMode::Spsc> OutgoingMessages;
+	FEvent* Semaphore;
+	FCaptureMessage CurrentMessage;
+	bool bCameraConnected;
+	const float RetryConnectWaitTime;
+
+#if WITH_MLSDK
+	MLCameraDeviceStatusCallbacks DeviceStatusCallbacks;
+	MLCameraCaptureCallbacks CaptureCallbacks;
+	MLCameraOutput* CameraOutput;
+#endif //WITH_MLSDK
+
+	const FString ImgExtension;
+	const FString VidExtension;
+	FString UniqueFileName;
+	TSharedPtr<IImageWrapper> ImageWrapper;
+	FThreadSafeBool bPaused;
+	static FThreadSafeCounter64 PreviewHandle;
+
+private:
+#if WITH_MLSDK
+	static void OnPreviewBufferAvailable(MLHandle Output, void *Data);
+#endif //WITH_MLSDK
+};
+
+DECLARE_LOG_CATEGORY_EXTERN(LogCameraCaptureRunnable, Verbose, All);
\ No newline at end of file
