#DEV-VR - Merge from //UE4/Dev-MagicLeap/... @ CL 4136411

diff --git a/Engine/Plugins/Lumin/MagicLeap/Source/MagicLeap/Private/ImageTrackerComponent.cpp b/Engine/Plugins/Lumin/MagicLeap/Source/MagicLeap/Private/ImageTrackerComponent.cpp
--- a/Engine/Plugins/Lumin/MagicLeap/Source/MagicLeap/Private/ImageTrackerComponent.cpp
+++ b/Engine/Plugins/Lumin/MagicLeap/Source/MagicLeap/Private/ImageTrackerComponent.cpp
@@ -1,17 +1,13 @@
 // Copyright 1998-2018 Epic Games, Inc. All Rights Reserved.
 
 #include "ImageTrackerComponent.h"
+#include "ImageTrackerRunnable.h"
 #include "MagicLeapHMD.h"
 #include "AppFramework.h"
 #include "MagicLeapMath.h"
 #include "Runtime/Core/Public/HAL/PlatformFilemanager.h"
 #include "Misc/Paths.h"
 #include "AppEventHandler.h"
-#include "HAL/Runnable.h"
-#include "HAL/RunnableThread.h"
-#if PLATFORM_LUMIN
-#include "Lumin/LuminAffinity.h"
-#endif // PLATFORM_LUMIN
 #if WITH_EDITOR
 #include "Editor.h"
 #endif
@@ -20,346 +16,251 @@
 #if WITH_MLSDK
 ML_INCLUDES_START
 #include <ml_image_tracking.h>
-#include <ml_perception.h>
-#include <ml_snapshot.h>
-#include <ml_head_tracking.h>
-#include <ml_coordinate_frame_uid.h>
 ML_INCLUDES_END
 #endif //WITH_MLSDK
 
-class FImageTrackerEngineInterface : public MagicLeap::IAppEventHandler
+class FImageTrackerImpl : public MagicLeap::IAppEventHandler
 {
 public:
-	static TWeakPtr<FImageTrackerEngineInterface, ESPMode::ThreadSafe> Get()
+	FImageTrackerImpl(UImageTrackerComponent* InOwner)
+	: Owner(InOwner)
+	, State(EState::WaitingForAppInit)
+#if WITH_MLSDK
+	, Target(ML_INVALID_HANDLE)
+#endif //WITH_MLSDK
+#if WITH_EDITOR
+	, TextureBeforeEdit(nullptr)
+#endif
+	, PrevTextureTarget(nullptr)
+	{
+#if WITH_MLSDK
+		OldTrackingStatus.status = MLImageTrackerTargetStatus_Ensure32Bits;
+#endif //WITH_MLSDK
+	};
+
+	virtual ~FImageTrackerImpl()
 	{
-		if (!Instance.IsValid())
-		{
-			Instance = MakeShareable(new FImageTrackerEngineInterface());
-		}
-		// Because this is a singleton, when running in ZI, consecutive VRPreview launches would fail
-		// because the tracker would be destroyed but created only on the very first launch, when
-		// FImageTrackerEngineInterface was constructed.
-		// Thus, we need to create the tracker separately.
 #if WITH_MLSDK
-		if (!MLHandleIsValid(Instance->GetHandle()))
+		if (MLHandleIsValid(Target))
 		{
-			Instance->CreateTracker();
+			checkf(ImageTrackerRunnable.IsValid(), TEXT("[FImageTrackerImpl] ImageTracker weak pointer is invalid!"));
+			MLResult Result = MLImageTrackerRemoveTarget(ImageTrackerRunnable->GetHandle(), Target);
+			UE_CLOG(Result != MLResult_Ok, LogMagicLeap, Error, TEXT("MLImageTrackerRemoveTarget failed with error %d."), Result);
+			Target = ML_INVALID_HANDLE;
 		}
 #endif //WITH_MLSDK
+		// Calling this explicitly to make the chain of destruction more obvious,
+		// ie a potential call to FImageTrackerRunnable's destructor right here.
+		// On the lumin platform this call will take place inside the destruction
+		// worker thread.
+		ImageTrackerRunnable.Reset();
+
+		// need this call to prevent the runnable form persisting with a ref count of 1 (due to Appframework's shared pointer instance).
+		if (GEngine && GEngine->XRSystem.IsValid() && GEngine->XRSystem->GetHMDDevice())
+		{
+			static_cast<FMagicLeapHMD*>(GEngine->XRSystem->GetHMDDevice())->GetAppFramework().RefreshImageTrackerRunnableReferences();
+		}
 
-		return Instance;
 	}
 
-#if WITH_MLSDK
-	MLHandle GetHandle() const
+	void SetTargetAsync(const FString& InName, bool bInIsStationary, float InLongerDimension, UTexture2D* InTargetTexture)
 	{
-		return ImageTracker;
-	}
-#endif //WITH_MLSDK
-
-private:
-	FImageTrackerEngineInterface()
 #if WITH_MLSDK
-	: ImageTracker(ML_INVALID_HANDLE)
+		checkf(ImageTrackerRunnable.IsValid(), TEXT("Invalid image tracker!"));
+		FTrackerMessage CreateTargetMsg;
+		CreateTargetMsg.TaskType = FTrackerMessage::TaskType::TryCreateTarget;
+		CreateTargetMsg.Requester = this;
+		CreateTargetMsg.PrevTarget = Target;
+		CreateTargetMsg.TargetName = InName;
+		MLImageTrackerTargetSettings TargetSettings;
+		TargetSettings.longer_dimension = InLongerDimension;
+		TargetSettings.is_stationary = bInIsStationary;
+		CreateTargetMsg.TargetSettings = TargetSettings;
+		CreateTargetMsg.TargetImageTexture = InTargetTexture;
+		ImageTrackerRunnable->IncomingMessages.Enqueue(CreateTargetMsg);
 #endif //WITH_MLSDK
-	{
-		CreateTracker();
 	}
 
-	void CreateTracker()
+	void Tick()
 	{
 #if WITH_MLSDK
-		UE_LOG(LogMagicLeap, Display, TEXT("[FImageTrackerEngineInterface] Creating Image Tracker"));
-		FMemory::Memset(&Settings, 0, sizeof(Settings));
-		bool bResult = MLImageTrackerInitSettings(&Settings);
-
-		if (!bResult)
+		if (!IMagicLeapPlugin::Get().IsMagicLeapHMDValid())
 		{
-			UE_LOG(LogMagicLeap, Error, TEXT("[FImageTrackerEngineInterface] Could not initialize image tracker settings mofo"));
+			return;
 		}
 
-		ImageTracker = MLImageTrackerCreate(&Settings);
-
-		if (!MLHandleIsValid(ImageTracker))
+		switch (State)
 		{
-			UE_LOG(LogMagicLeap, Error, TEXT("[FImageTrackerEngineInterface] Could not create Image tracker."));
-		}
-#endif //WITH_MLSDK
-	}	
+		case WaitingForAppInit:
+		{
+			checkf(GEngine, TEXT("[FCameraCaptureRunnable::Get()] GEngine is null!"));
+			FAppFramework& AppFramework = StaticCastSharedPtr<FMagicLeapHMD>(GEngine->XRSystem)->GetAppFramework();
 
-	void OnAppPause() override
-	{
-#if WITH_MLSDK
-		bWasSystemEnabledOnPause = Settings.enable_image_tracking;
+			if (!AppFramework.IsInitialized())
+			{
+				UE_LOG(LogMagicLeap, Warning, TEXT("[UImageTrackerComponent] AppFramework not initialized."));
+			}
+			else
+			{
+				State = NoTarget;
+				ImageTrackerRunnable = AppFramework.GetImageTrackerRunnable();
+			}
 
-		if (!bWasSystemEnabledOnPause)
-		{
-			UE_LOG(LogMagicLeap, Log, TEXT("[FImageTrackerEngineInterface] Image tracking was not enabled at time of application pause."));
+			break;
 		}
-		else
+		case NoTarget:
 		{
-			if (!MLHandleIsValid(ImageTracker))
+			if (Owner->TargetImageTexture == nullptr)
 			{
-				UE_LOG(LogMagicLeap, Error, TEXT("[FImageTrackerEngineInterface] Image tracker was invalid on application pause."));
+				UE_LOG(LogMagicLeap, Warning, TEXT("[UImageTrackerComponent] ImageTracker: No image selected to track."));
+			}
+			else if (Owner->TargetImageTexture->GetPixelFormat() != EPixelFormat::PF_R8G8B8A8 && Owner->TargetImageTexture->GetPixelFormat() != EPixelFormat::PF_B8G8R8A8)
+			{
+				State = TargetCreationError;
+				UE_LOG(LogMagicLeap, Error, TEXT("[UImageTrackerComponent] ImageTracker: Unsupported pixel format encountered!"));
 			}
 			else
 			{
-				Settings.enable_image_tracking = false;
+				State = CreatingTarget;
+				if (Owner->Name.Len() == 0) Owner->Name = Owner->GetName();
+				const FAppFramework& AppFramework = StaticCastSharedPtr<FMagicLeapHMD>(GEngine->XRSystem)->GetAppFrameworkConst();
+				SetTargetAsync(Owner->Name, Owner->bIsStationary, Owner->LongerDimension / AppFramework.GetWorldToMetersScale(), Owner->TargetImageTexture);
+			}
 
-				if (!MLImageTrackerUpdateSettings(ImageTracker, &Settings))
-				{
-					UE_LOG(LogMagicLeap, Error, TEXT("[FImageTrackerEngineInterface] Failed to disable image tracker on application pause."));
-				}
-				else
+			break;
+		}
+		case CreatingTarget:
+		{
+			FTrackerMessage OutMsg;
+			if (ImageTrackerRunnable->OutgoingMessages.Peek(OutMsg))
+			{
+				if (OutMsg.Requester == this)
 				{
-					UE_LOG(LogMagicLeap, Log, TEXT("[FImageTrackerEngineInterface] Image tracker paused until app resumes."));
+					ImageTrackerRunnable->OutgoingMessages.Pop();
+					if (OutMsg.TaskType == FTrackerMessage::TargetCreateSucceeded)
+					{
+						State = Tracking;
+						Data = OutMsg.Data;
+						Target = OutMsg.Target;
+						PrevTextureTarget = Owner->TargetImageTexture;
+						Owner->OnSetImageTargetSucceeded.Broadcast();
+					}
+					else
+					{
+						State = TargetCreationError;
+						Owner->OnSetImageTargetFailed.Broadcast();
+					}
 				}
 			}
-		}
-#endif //WITH_MLSDK
-	}
 
-	void OnAppResume() override
-	{
-#if WITH_MLSDK
-		if (!bWasSystemEnabledOnPause)
+			break;
+		}
+		case TargetCreationError:
 		{
-			UE_LOG(LogMagicLeap, Log, TEXT("[FImageTrackerEngineInterface] Not resuming image tracker as it was not enabled at time of application pause."));
+			CheckForNewTarget();
+			break;
 		}
-		else
+		case Tracking:
 		{
-			if (!MLHandleIsValid(ImageTracker))
+			checkf(ImageTrackerRunnable.IsValid(), TEXT("[FImageTrackerImpl] ImageTracker weak pointer is invalid!"));
+			MLImageTrackerTargetResult TrackingStatus;
+			MLResult Result = MLImageTrackerGetTargetResult(ImageTrackerRunnable->GetHandle(), Target, &TrackingStatus);
+			if (Result != MLResult_Ok)
 			{
-				UE_LOG(LogMagicLeap, Error, TEXT("[FImageTrackerEngineInterface] Image tracker was invalid on application resume."));
+				TrackingStatus.status = MLImageTrackerTargetStatus_NotTracked;
+			}
+
+			if (TrackingStatus.status == MLImageTrackerTargetStatus_NotTracked)
+			{
+				if (OldTrackingStatus.status != MLImageTrackerTargetStatus_NotTracked)
+				{
+					Owner->OnImageTargetLost.Broadcast();
+				}
 			}
 			else
 			{
-				Settings.enable_image_tracking = true;
-
-				if (!MLImageTrackerUpdateSettings(ImageTracker, &Settings))
+				EFailReason FailReason = EFailReason::None;
+				FTransform Pose = FTransform::Identity;
+				const FAppFramework& AppFramework = StaticCastSharedPtr<FMagicLeapHMD>(GEngine->XRSystem)->GetAppFrameworkConst();
+				if (AppFramework.GetTransform(Data.coord_frame_target, Pose, FailReason))
 				{
-					UE_LOG(LogMagicLeap, Error, TEXT("[FImageTrackerEngineInterface] Failed to re-enable image tracker on application resume."));
+					Pose.SetRotation(MagicLeap::ToUERotator(Pose.GetRotation()));
+					if (TrackingStatus.status == MLImageTrackerTargetStatus_Unreliable)
+					{
+						FVector LastTrackedLocation = Owner->GetComponentLocation();
+						FRotator LastTrackedRotation = Owner->GetComponentRotation();
+						if (Owner->bUseUnreliablePose)
+						{
+							Owner->SetRelativeLocationAndRotation(Pose.GetTranslation(), Pose.Rotator());
+						}
+						// Developer can choose whether to use this unreliable pose or not.
+						Owner->OnImageTargetUnreliableTracking.Broadcast(LastTrackedLocation, LastTrackedRotation, Pose.GetTranslation(), Pose.Rotator());
+					}
+					else
+					{
+						Owner->SetRelativeLocationAndRotation(Pose.GetTranslation(), Pose.Rotator());
+						if (OldTrackingStatus.status != MLImageTrackerTargetStatus_Tracked)
+						{
+							Owner->OnImageTargetFound.Broadcast();
+						}
+					}
 				}
 				else
 				{
-					UE_LOG(LogMagicLeap, Log, TEXT("[FImageTrackerEngineInterface] Image tracker re-enabled on application resume."));
+					if (FailReason == EFailReason::NaNsInTransform)
+					{
+						UE_LOG(LogMagicLeap, Error, TEXT("[UImageTrackerComponent] NaNs in image tracker target transform."));
+					}
+					TrackingStatus.status = MLImageTrackerTargetStatus_NotTracked;
+					if (OldTrackingStatus.status != MLImageTrackerTargetStatus_NotTracked)
+					{
+						Owner->OnImageTargetLost.Broadcast();
+					}
 				}
 			}
-		}
-#endif //WITH_MLSDK
-	}
-
-	void OnAppShutDown() override
-	{
-#if WITH_MLSDK
-		if (MLHandleIsValid(ImageTracker))
-		{
-			bool bResult = MLImageTrackerDestroy(ImageTracker);
-
-			if (!bResult)
-			{
-				UE_LOG(LogMagicLeap, Error, TEXT("[FImageTrackerEngineInterface] Error destroying image tracker."));
-			}
-
-			ImageTracker = ML_INVALID_HANDLE;
-		}
-#endif //WITH_MLSDK
-	}
-
-#if WITH_MLSDK
-	MLHandle ImageTracker;
-	MLImageTrackerSettings Settings;
-#endif //WITH_MLSDK
-	static TSharedPtr<FImageTrackerEngineInterface, ESPMode::ThreadSafe> Instance;
-};
-
-TSharedPtr<FImageTrackerEngineInterface, ESPMode::ThreadSafe> FImageTrackerEngineInterface::Instance;
-
-struct FTrackerMessage
-{
-	enum TaskType
-	{
-		None,
-		Create,
-		ReportStatus,
-	};
-
-	TaskType TaskType;
-	FString TargetName;
-	UTexture2D* TargetImageTexture;
-#if WITH_MLSDK
-	MLImageTrackerTargetSettings TargetSettings;
-	MLImageTrackerTargetResult TrackingStatus;
-#endif //WITH_MLSDK
-
-	FTrackerMessage()
-	: TaskType(TaskType::None)
-	{}
-};
-
-class FImageTrackerImpl : public FRunnable, public MagicLeap::IAppEventHandler
-{
-public:
-	FImageTrackerImpl()
-		: bHasTarget(false)
-#if WITH_EDITOR
-		, TextureBeforeEdit(nullptr)
-#endif
-#if WITH_MLSDK
-		, Target(ML_INVALID_HANDLE)
-#endif //WITH_MLSDK
-		, Thread(nullptr)
-		, StopTaskCounter(0)
-		, bIsTracking(false)
-	{
-#if WITH_MLSDK
-		OldTrackingStatus.status = MLImageTrackerTargetStatus_Ensure32Bits;
-#endif //WITH_MLSDK
-
-#if PLATFORM_LUMIN
-		Thread = FRunnableThread::Create(this, TEXT("ImageTrackerWorker"), 0, TPri_BelowNormal, FLuminAffinity::GetPoolThreadMask());
-#else
-		Thread = FRunnableThread::Create(this, TEXT("ImageTrackerWorker"), 0, TPri_BelowNormal);
-#endif // PLATFORM_LUMIN
-	};
-
-	virtual ~FImageTrackerImpl()
-	{
-		StopTaskCounter.Increment();
 
-		if (nullptr != Thread)
-		{
-			Thread->WaitForCompletion();
-			delete Thread;
-			Thread = nullptr;
+			OldTrackingStatus = TrackingStatus;
+			CheckForNewTarget();
+			break;
 		}
-
-#if WITH_MLSDK
-		if (MLHandleIsValid(Target))
-		{
-			checkf(ImageTracker.Pin().IsValid(), TEXT("[FImageTrackerImpl] ImageTracker weak pointer is invalid!"));
-			MLImageTrackerRemoveTarget(ImageTracker.Pin()->GetHandle(), Target);
-			Target = ML_INVALID_HANDLE;
 		}
 #endif //WITH_MLSDK
-
-		ImageTracker.Reset();
 	}
 
-	virtual uint32 Run() override
+	void CheckForNewTarget()
 	{
-		while (StopTaskCounter.GetValue() == 0)
+		if (Owner->TargetImageTexture != PrevTextureTarget)
 		{
-			if (IncomingMessages.Dequeue(CurrentMessage))
-			{
-				DoTasks();
-			}
-
-			FPlatformProcess::Sleep(0.5f);
+			State = NoTarget;
 		}
-
-		return 0;
 	}
 
-	void SetTargetAsync(const FString& InName, bool bInIsStationary, float InLongerDimension, UTexture2D* InTargetTexture)
+public:
+	enum EState
 	{
-#if WITH_MLSDK
-		if (!ImageTracker.IsValid())
-		{
-			ImageTracker = FImageTrackerEngineInterface::Get();
-		}
+		WaitingForAppInit,
+		NoTarget,
+		CreatingTarget,
+		TargetCreationError,
+		Tracking,
+	};
 
-		bHasTarget = true;
-		MLImageTrackerTargetSettings TargetSettings;
-		TargetSettings.longer_dimension = InLongerDimension;
-		TargetSettings.is_stationary = bInIsStationary;
-		FTrackerMessage CreateTargetMsg;
-		CreateTargetMsg.TaskType = FTrackerMessage::TaskType::Create;
-		CreateTargetMsg.TargetName = InName;
-		CreateTargetMsg.TargetSettings = TargetSettings;
-		CreateTargetMsg.TargetImageTexture = InTargetTexture;
-		IncomingMessages.Enqueue(CreateTargetMsg);
-#endif //WITH_MLSDK
-	}
-	
-public:
-	TQueue<FTrackerMessage, EQueueMode::Spsc> OutgoingMessages;
-	bool bHasTarget;
-#if WITH_EDITOR
-	UTexture2D* TextureBeforeEdit;
-#endif
+	UImageTrackerComponent* Owner;
+	EState State;
+	TSharedPtr<FImageTrackerRunnable, ESPMode::ThreadSafe> ImageTrackerRunnable;
 #if WITH_MLSDK
 	MLHandle Target;
 	MLImageTrackerTargetStaticData Data;
 	MLImageTrackerTargetResult OldTrackingStatus;
 #endif //WITH_MLSDK
-
-public:
-	TWeakPtr<FImageTrackerEngineInterface, ESPMode::ThreadSafe> ImageTracker;
-	FRunnableThread* Thread;
-	FThreadSafeCounter StopTaskCounter;
-	TQueue<FTrackerMessage, EQueueMode::Spsc> IncomingMessages;
-	FTrackerMessage CurrentMessage;
-	bool bIsTracking;
-
-	void DoTasks()
-	{
-		switch (CurrentMessage.TaskType)
-		{
-		case FTrackerMessage::TaskType::None: break;
-		case FTrackerMessage::TaskType::Create: SetTarget(); break;
-		case FTrackerMessage::TaskType::ReportStatus: break;
-		}
-	}
-
-	void SetTarget()
-	{
-#if WITH_MLSDK
-		if (!MLHandleIsValid(Target))
-		{
-			CurrentMessage.TargetSettings.name = TCHAR_TO_UTF8(*CurrentMessage.TargetName);
-			FTexture2DMipMap& Mip = CurrentMessage.TargetImageTexture->PlatformData->Mips[0];
-			const unsigned char* PixelData = static_cast<const unsigned char*>(Mip.BulkData.Lock(LOCK_READ_ONLY));
-			checkf(ImageTracker.Pin().IsValid(), TEXT("[FImageTrackerImpl] ImageTracker weak pointer is invalid!"));
-			Target = MLImageTrackerAddTargetFromArray(ImageTracker.Pin()->GetHandle(), &CurrentMessage.TargetSettings, PixelData, CurrentMessage.TargetImageTexture->GetSurfaceWidth(), CurrentMessage.TargetImageTexture->GetSurfaceHeight(), MLImageTrackerImageFormat_RGBA);
-			Mip.BulkData.Unlock();
-
-			if (!MLHandleIsValid(Target))
-			{
-				UE_LOG(LogMagicLeap, Error, TEXT("[FImageTrackerImpl] Could not create Image Target."));
-				return;
-			}
-
-			// [3] Cache all the static data for this target.
-			checkf(ImageTracker.Pin().IsValid(), TEXT("[FImageTrackerImpl] ImageTracker weak pointer is invalid!"));
-			if (!MLImageTrackerGetTargetStaticData(ImageTracker.Pin()->GetHandle(), Target, &Data))
-			{
-				UE_LOG(LogMagicLeap, Error, TEXT("[FImageTrackerImpl] Could not get the static data for the Image Target."));
-				return;
-			}
-
-			bIsTracking = true;
-		}
-#endif //WITH_MLSDK
-	}
-
-public:
-#if WITH_MLSDK
-	inline MLHandle Tracker()
-	{
-		if (!ImageTracker.IsValid())
-		{
-			return ML_INVALID_HANDLE;
-		}
-		checkf(ImageTracker.Pin().IsValid(), TEXT("[FImageTrackerImpl] ImageTracker weak pointer is invalid!"));
-		return ImageTracker.Pin()->GetHandle();
-	}
-#endif //WITH_MLSDK
+#if WITH_EDITOR
+	UTexture2D* TextureBeforeEdit;
+#endif
+	UTexture2D* PrevTextureTarget;
 };
 
 UImageTrackerComponent::UImageTrackerComponent()
-	: Impl(new FImageTrackerImpl())
-	, bTick(true)
+: Impl(new FImageTrackerImpl(this))
 {
 	// Make sure this component ticks
 	PrimaryComponentTick.bCanEverTick = true;
@@ -385,98 +286,21 @@ UImageTrackerComponent::~UImageTrackerComponent()
 
 void UImageTrackerComponent::TickComponent(float DeltaTime, enum ELevelTick TickType, FActorComponentTickFunction *ThisTickFunction)
 {
-#if WITH_MLSDK && PLATFORM_LUMIN
 	Super::TickComponent(DeltaTime, TickType, ThisTickFunction);
+	Impl->Tick();
+}
 
-	const FAppFramework& AppFramework = StaticCastSharedPtr<FMagicLeapHMD>(GEngine->XRSystem)->GetAppFrameworkConst();
-
-	if (!AppFramework.IsInitialized())
-	{
-		UE_LOG(LogMagicLeap, Warning, TEXT("[UImageTrackerComponent] AppFramework not initialized."));
-		return;
-	}
-
-	if (TargetImageTexture == nullptr)
-	{
-		UE_LOG(LogMagicLeap, Warning, TEXT("ImageTracker: No image selected to track."));
-		return;
-	}
-
-	if (TargetImageTexture->GetPixelFormat() != EPixelFormat::PF_R8G8B8A8 && TargetImageTexture->GetPixelFormat() != EPixelFormat::PF_B8G8R8A8)
+void UImageTrackerComponent::SetTargetAsync(UTexture2D* NewTextureTarget)
+{
+	if (!NewTextureTarget || (NewTextureTarget->GetPixelFormat() == EPixelFormat::PF_R8G8B8A8 || NewTextureTarget->GetPixelFormat() == EPixelFormat::PF_B8G8R8A8))
 	{
-		UE_LOG(LogMagicLeap, Error, TEXT("[UImageTrackerComponent] ImageTracker: Unsupported pixel format encountered!"));
-		bTick = false;
+		Impl->PrevTextureTarget = TargetImageTexture;
+		TargetImageTexture = NewTextureTarget;
 	}
-
-	if (!bTick) return;
-
-	if (!Impl->bHasTarget)
+	else
 	{
-		if (Name.Len() == 0) Name = GetName();
-		Impl->SetTargetAsync(Name, bIsStationary, LongerDimension / AppFramework.GetWorldToMetersScale(), TargetImageTexture);
+		UE_LOG(LogMagicLeap, Error, TEXT("[UImageTrackerComponent::SetTarget()] ImageTracker: Unsupported pixel format encountered!"));
 	}
-
-	if (Impl->bIsTracking)
-	{
-		checkf(Impl->ImageTracker.Pin().IsValid(), TEXT("[FImageTrackerImpl] ImageTracker weak pointer is invalid!"));
-
-		MLImageTrackerTargetResult TrackingStatus;
-		if (!MLImageTrackerGetTargetResult(Impl->ImageTracker.Pin()->GetHandle(), Impl->Target, &TrackingStatus))
-		{
-			TrackingStatus.status = MLImageTrackerTargetStatus_NotTracked;
-		}
-
-		if (TrackingStatus.status == MLImageTrackerTargetStatus_NotTracked)
-			{
-				if (Impl->OldTrackingStatus.status != MLImageTrackerTargetStatus_NotTracked)
-				{
-					OnImageTargetLost.Broadcast();
-				}
-			}
-			else
-			{
-				EFailReason FailReason = EFailReason::None;
-				FTransform Pose = FTransform::Identity;
-				if (AppFramework.GetTransform(Impl->Data.coord_frame_target, Pose, FailReason))
-				{
-				Pose.SetRotation(MagicLeap::ToUERotator(Pose.GetRotation()));
-				if (TrackingStatus.status == MLImageTrackerTargetStatus_Unreliable)
-					{
-						FVector LastTrackedLocation = GetComponentLocation();
-						FRotator LastTrackedRotation = GetComponentRotation();
-						if (bUseUnreliablePose)
-						{
-							this->SetRelativeLocationAndRotation(Pose.GetTranslation(), Pose.Rotator());
-						}
-						// Developer can choose whether to use this unreliable pose or not.
-						OnImageTargetUnreliableTracking.Broadcast(LastTrackedLocation, LastTrackedRotation, Pose.GetTranslation(), Pose.Rotator());
-					}
-					else
-					{
-						this->SetRelativeLocationAndRotation(Pose.GetTranslation(), Pose.Rotator());
-						if (Impl->OldTrackingStatus.status != MLImageTrackerTargetStatus_Tracked)
-						{
-							OnImageTargetFound.Broadcast();
-						}
-					}
-				}
-				else
-				{
-					if (FailReason == EFailReason::NaNsInTransform)
-					{
-						UE_LOG(LogMagicLeap, Error, TEXT("[UImageTrackerComponent] NaNs in image tracker target transform."));
-					}
-				TrackingStatus.status = MLImageTrackerTargetStatus_NotTracked;
-					if (Impl->OldTrackingStatus.status != MLImageTrackerTargetStatus_NotTracked)
-					{
-						OnImageTargetLost.Broadcast();
-					}
-				}
-			}
-
-		Impl->OldTrackingStatus = TrackingStatus;
-	}
-#endif //WITH_MLSDK
 }
 
 #if WITH_EDITOR
@@ -498,6 +322,7 @@ void UImageTrackerComponent::PostEditChangeProperty(FPropertyChangedEvent& Prope
 	{
 		if (TargetImageTexture && !(TargetImageTexture->GetPixelFormat() == EPixelFormat::PF_R8G8B8A8 || TargetImageTexture->GetPixelFormat() == EPixelFormat::PF_B8G8R8A8))
 		{
+			Impl->PrevTextureTarget = TargetImageTexture;
 			TargetImageTexture = Impl->TextureBeforeEdit;
 			UE_LOG(LogMagicLeap, Error, TEXT("[UImageTrackerComponent] Cannot set texture %s as it uses an invalid pixel format!  Valid formats are R8B8G8A8 or B8G8R8A8"), *TargetImageTexture->GetName());
 		}
@@ -505,4 +330,26 @@ void UImageTrackerComponent::PostEditChangeProperty(FPropertyChangedEvent& Prope
 
 	Super::PostEditChangeProperty(PropertyChangedEvent);
 }
-#endif
+#endif //WITH_EDITOR
+
+void UImageTrackerFunctionLibrary::SetMaxSimultaneousTargets(int32 MaxSimultaneousTargets)
+{
+	FTrackerMessage SetMaxTargetsMsg;
+	SetMaxTargetsMsg.TaskType = FTrackerMessage::TaskType::SetMaxTargets;
+	SetMaxTargetsMsg.MaxTargets = MaxSimultaneousTargets < 1 ? 1 : static_cast<uint32>(MaxSimultaneousTargets);
+	checkf(GEngine, TEXT("[UImageTrackerFunctionLibrary::SetMaxSimultaneousTargets()] GEngine is null!"));
+	FAppFramework& AppFramework = StaticCastSharedPtr<FMagicLeapHMD>(GEngine->XRSystem)->GetAppFramework();
+	checkf(AppFramework.IsInitialized(), TEXT("[UImageTrackerFunctionLibrary::SetMaxSimultaneousTargets()] AppFramework not yet initialized!"));
+	AppFramework.GetImageTrackerRunnable()->IncomingMessages.Enqueue(SetMaxTargetsMsg);
+}
+
+void UImageTrackerFunctionLibrary::EnableImageTracking(bool bEnable)
+{
+	FTrackerMessage EnableMsg;
+	EnableMsg.TaskType = FTrackerMessage::TaskType::SetEnabled;
+	EnableMsg.bEnable = bEnable;
+	checkf(GEngine, TEXT("[UImageTrackerFunctionLibrary::EnableImageTracking()] GEngine is null!"));
+	FAppFramework& AppFramework = StaticCastSharedPtr<FMagicLeapHMD>(GEngine->XRSystem)->GetAppFramework();
+	checkf(AppFramework.IsInitialized(), TEXT("[UImageTrackerFunctionLibrary::EnableImageTracking()] AppFramework not yet initialized!"));
+	AppFramework.GetImageTrackerRunnable()->IncomingMessages.Enqueue(EnableMsg);
+}
\ No newline at end of file
diff --git a/Engine/Plugins/Lumin/MagicLeap/Source/MagicLeap/Private/ImageTrackerRunnable.cpp b/Engine/Plugins/Lumin/MagicLeap/Source/MagicLeap/Private/ImageTrackerRunnable.cpp
new file mode 100644
--- /dev/null
+++ b/Engine/Plugins/Lumin/MagicLeap/Source/MagicLeap/Private/ImageTrackerRunnable.cpp
@@ -0,0 +1,378 @@
+// %BANNER_BEGIN%
+// ---------------------------------------------------------------------
+// %COPYRIGHT_BEGIN%
+//
+// Copyright (c) 2017 Magic Leap, Inc. (COMPANY) All Rights Reserved.
+// Magic Leap, Inc. Confidential and Proprietary
+//
+// NOTICE: All information contained herein is, and remains the property
+// of COMPANY. The intellectual and technical concepts contained herein
+// are proprietary to COMPANY and may be covered by U.S. and Foreign
+// Patents, patents in process, and are protected by trade secret or
+// copyright law. Dissemination of this information or reproduction of
+// this material is strictly forbidden unless prior written permission is
+// obtained from COMPANY. Access to the source code contained herein is
+// hereby forbidden to anyone except current COMPANY employees, managers
+// or contractors who have executed Confidentiality and Non-disclosure
+// agreements explicitly covering such access.
+//
+// The copyright notice above does not evidence any actual or intended
+// publication or disclosure of this source code, which includes
+// information that is confidential and/or proprietary, and is a trade
+// secret, of COMPANY. ANY REPRODUCTION, MODIFICATION, DISTRIBUTION,
+// PUBLIC PERFORMANCE, OR PUBLIC DISPLAY OF OR THROUGH USE OF THIS
+// SOURCE CODE WITHOUT THE EXPRESS WRITTEN CONSENT OF COMPANY IS
+// STRICTLY PROHIBITED, AND IN VIOLATION OF APPLICABLE LAWS AND
+// INTERNATIONAL TREATIES. THE RECEIPT OR POSSESSION OF THIS SOURCE
+// CODE AND/OR RELATED INFORMATION DOES NOT CONVEY OR IMPLY ANY RIGHTS
+// TO REPRODUCE, DISCLOSE OR DISTRIBUTE ITS CONTENTS, OR TO MANUFACTURE,
+// USE, OR SELL ANYTHING THAT IT MAY DESCRIBE, IN WHOLE OR IN PART.
+//
+// %COPYRIGHT_END%
+// --------------------------------------------------------------------
+// %BANNER_END%
+
+#include "ImageTrackerRunnable.h"
+#include "MagicLeapHMD.h"
+#include "AppFramework.h"
+#include "MagicLeapMath.h"
+#include "Runtime/Core/Public/HAL/PlatformFilemanager.h"
+#include "Misc/Paths.h"
+#include "Engine/Texture2D.h"
+#if PLATFORM_LUMIN
+#include "Lumin/LuminAffinity.h"
+#endif // PLATFORM_LUMIN
+
+#if WITH_MLSDK
+ML_INCLUDES_START
+//#include <ml_image_tracking.h>
+#include <ml_perception.h>
+#include <ml_snapshot.h>
+#include <ml_head_tracking.h>
+#include <ml_coordinate_frame_uid.h>
+ML_INCLUDES_END
+#endif //WITH_MLSDK
+
+FImageTrackerRunnable::FImageTrackerRunnable()
+#if WITH_MLSDK
+: MagicLeap::IAppEventHandler({ MLPrivilegeID_CameraCapture })
+, ImageTracker(ML_INVALID_HANDLE)
+, Thread(nullptr)
+#else
+: Thread(nullptr)
+#endif //WITH_MLSDK
+, StopTaskCounter(0)
+, RetryCreateTrackerWaitTime(0.5f)
+{
+#if WITH_MLSDK
+	MLResult Result = MLImageTrackerInitSettings(&Settings);
+	UE_CLOG(Result != MLResult_Ok, LogMagicLeap, Error, TEXT("MLImageTrackerInitSettings failed with error %d."), Result);
+#endif //WITH_MLSDK
+#if PLATFORM_LUMIN
+	Thread = FRunnableThread::Create(this, TEXT("ImageTrackerWorker"), 0, TPri_BelowNormal, FLuminAffinity::GetPoolThreadMask());
+#else
+	Thread = FRunnableThread::Create(this, TEXT("ImageTrackerWorker"), 0, TPri_BelowNormal);
+#endif // PLATFORM_LUMIN
+}
+
+FImageTrackerRunnable::~FImageTrackerRunnable()
+{
+	StopTaskCounter.Increment();
+
+	if (nullptr != Thread)
+	{
+		Thread->WaitForCompletion();
+		delete Thread;
+		Thread = nullptr;
+	}
+}
+
+#if WITH_MLSDK
+MLHandle FImageTrackerRunnable::GetHandle() const
+{
+	return ImageTracker;
+}
+#endif //WITH_MLSDK
+
+uint32 FImageTrackerRunnable::Run()
+{
+	// JMC: hack to prevent privilege requests before service is properly initialised.
+	FPlatformProcess::Sleep(5.0f);
+	while (StopTaskCounter.GetValue() == 0)
+	{
+#if WITH_MLSDK
+		if (!MLHandleIsValid(ImageTracker))
+		{
+			if (GetPrivilegeStatus(MLPrivilegeID_CameraCapture) == MagicLeap::EPrivilegeState::Granted)
+			{
+				UE_LOG(LogMagicLeap, Display, TEXT("[FImageTrackerRunnable] Attempting to create image tracker."));
+
+				MLResult Result = MLImageTrackerCreate(&Settings, &ImageTracker);
+				if (Result != MLResult_Ok)
+				{
+					UE_LOG(LogMagicLeap, Error, TEXT("MLImageTrackerCreate failed with error %s."), UTF8_TO_TCHAR(MLGetResultString(Result)));
+					FPlatformProcess::Sleep(RetryCreateTrackerWaitTime);
+				}
+			}
+			else
+			{
+				UE_LOG(LogMagicLeap, Error, TEXT("[FImageTrackerRunnable] Failed to create image tracker due to lack of privilege!"));
+			}
+		}
+		else if (IncomingMessages.Dequeue(CurrentMessage))
+		{
+			switch (CurrentMessage.TaskType)
+			{
+			case FTrackerMessage::None: checkf(false, TEXT("Invalid incoming task 'FTrackerMessage::None'!"));	break;
+			case FTrackerMessage::Pause: TryPause(); break;
+			case FTrackerMessage::Resume: TryResume(); break;
+			case FTrackerMessage::SetEnabled: SetEnabled(); break;
+			case FTrackerMessage::SetMaxTargets: SetMaxTargets(); break;
+			case FTrackerMessage::TryCreateTarget: SetTarget(); break;
+			case FTrackerMessage::TargetCreateFailed: checkf(false, TEXT("Invalid incoming task 'FTrackerMessage::TargetCreateFailed'!"));	break;
+			case FTrackerMessage::TargetCreateSucceeded: checkf(false, TEXT("Invalid incoming task 'FTrackerMessage::TargetCreateSucceeded'!"));	break;
+			}
+		}
+#endif //WITH_MLSDK
+		FPlatformProcess::Sleep(0.5f);
+	}
+
+	return 0;
+}
+
+void FImageTrackerRunnable::OnAppPause()
+{
+	IAppEventHandler::OnAppPause();
+	FTrackerMessage PauseMsg;
+	PauseMsg.TaskType = FTrackerMessage::TaskType::Pause;
+	IncomingMessages.Enqueue(PauseMsg);
+}
+
+void FImageTrackerRunnable::OnAppResume()
+{
+	IAppEventHandler::OnAppResume();
+	FTrackerMessage ResumeMsg;
+	ResumeMsg.TaskType = FTrackerMessage::TaskType::Resume;
+	IncomingMessages.Enqueue(ResumeMsg);
+}
+
+void FImageTrackerRunnable::TryPause()
+{
+#if WITH_MLSDK
+	bWasSystemEnabledOnPause = Settings.enable_image_tracking;
+
+	if (!bWasSystemEnabledOnPause)
+	{
+		UE_LOG(LogMagicLeap, Log, TEXT("[FImageTrackerEngineInterface] Image tracking was not enabled at time of application pause."));
+	}
+	else
+	{
+		if (!MLHandleIsValid(ImageTracker))
+		{
+			UE_LOG(LogMagicLeap, Error, TEXT("[FImageTrackerEngineInterface] Image tracker was invalid on application pause."));
+		}
+		else
+		{
+			Settings.enable_image_tracking = false;
+			MLResult Result = MLImageTrackerUpdateSettings(ImageTracker, &Settings);
+			if (Result != MLResult_Ok)
+			{
+				UE_LOG(LogMagicLeap, Error, TEXT("[FImageTrackerEngineInterface] Failed to disable image tracker on application pause due to error %s."), UTF8_TO_TCHAR(MLGetResultString(Result)));
+			}
+			else
+			{
+				UE_LOG(LogMagicLeap, Log, TEXT("[FImageTrackerEngineInterface] Image tracker paused until app resumes."));
+			}
+		}
+	}
+#endif //WITH_MLSDK
+}
+
+void FImageTrackerRunnable::TryResume()
+{
+#if WITH_MLSDK
+	if (!bWasSystemEnabledOnPause)
+	{
+		UE_LOG(LogMagicLeap, Log, TEXT("[FImageTrackerRunnable] Not resuming image tracker as it was not enabled at time of application pause."));
+	}
+	else
+	{
+		if (!MLHandleIsValid(ImageTracker))
+		{
+			UE_LOG(LogMagicLeap, Error, TEXT("[FImageTrackerRunnable] Image tracker was invalid on application resume."));
+		}
+		else 
+		{
+			Settings.enable_image_tracking = true;
+			if (GetPrivilegeStatus(MLPrivilegeID_CameraCapture) == MagicLeap::EPrivilegeState::Granted)
+			{
+				MLResult Result = MLImageTrackerUpdateSettings(ImageTracker, &Settings);
+				if (Result != MLResult_Ok)
+				{
+					UE_LOG(LogMagicLeap, Error, TEXT("[FImageTrackerRunnable] Failed to re-enable image tracker on application resume due to error %s."), UTF8_TO_TCHAR(MLGetResultString(Result)));
+				}
+				else
+				{
+					UE_LOG(LogMagicLeap, Log, TEXT("[FImageTrackerRunnable] Image tracker re-enabled on application resume."));
+				}
+			}
+			else
+			{
+				UE_LOG(LogMagicLeap, Log, TEXT("[FImageTrackerRunnable] Image tracking failed to resume due to lack of privilege!"));
+			}
+		}
+	}
+#endif //WITH_MLSDK
+}
+
+void FImageTrackerRunnable::OnAppShutDown()
+{
+#if WITH_MLSDK
+	if (MLHandleIsValid(ImageTracker))
+	{
+		MLResult Result = MLImageTrackerDestroy(ImageTracker);
+		UE_CLOG(Result != MLResult_Ok, LogMagicLeap, Error, TEXT("MLImageTrackerDestroy failed with error %s."), UTF8_TO_TCHAR(MLGetResultString(Result)));
+		ImageTracker = ML_INVALID_HANDLE;
+	}
+#endif //WITH_MLSDK
+}
+
+void FImageTrackerRunnable::SetEnabled()
+{
+#if WITH_MLSDK
+	Settings.enable_image_tracking = CurrentMessage.bEnable;
+
+	UpdateTrackerSettings();
+#endif //WITH_MLSDK
+}
+
+void FImageTrackerRunnable::SetMaxTargets()
+{
+#if WITH_MLSDK
+	Settings.max_simultaneous_targets = CurrentMessage.MaxTargets;
+
+	UpdateTrackerSettings();
+#endif //WITH_MLSDK
+}
+
+void FImageTrackerRunnable::SetTarget()
+{
+#if WITH_MLSDK
+	if (MLHandleIsValid(CurrentMessage.PrevTarget))
+	{
+		MLResult Result = MLImageTrackerRemoveTarget(ImageTracker, CurrentMessage.PrevTarget);
+		UE_CLOG(Result != MLResult_Ok, LogMagicLeap, Error, TEXT("MLImageTrackerRemoveTarget failed with error %s."), UTF8_TO_TCHAR(MLGetResultString(Result)));
+	}
+
+	FTrackerMessage TargetCreateMsg;
+	TargetCreateMsg.Requester = CurrentMessage.Requester;
+
+	CurrentMessage.TargetSettings.name = TCHAR_TO_UTF8(*CurrentMessage.TargetName);
+	TArray<uint8*> MipPointers;
+	MipPointers.AddZeroed(CurrentMessage.TargetImageTexture->GetNumMips());
+	CurrentMessage.TargetImageTexture->GetMipData(0, (void**)MipPointers.GetData());
+	const int32 ImageWidth = CurrentMessage.TargetImageTexture->GetSizeX();
+	const int32 ImageHeight = CurrentMessage.TargetImageTexture->GetSizeY();
+	MLHandle Target = ML_INVALID_HANDLE;
+	if (MipPointers.Num() > 0 && MipPointers[0] != nullptr)
+	{
+		MLResult Result = MLImageTrackerAddTargetFromArray(
+			ImageTracker,
+			&CurrentMessage.TargetSettings,
+			MipPointers[0],
+			ImageWidth,
+			ImageHeight,
+			MLImageTrackerImageFormat_RGBA,
+			&Target);
+
+		if (Result != MLResult_Ok)
+		{
+			UE_LOG(LogMagicLeap, Error, TEXT("MLImageTrackerAddTargetFromArray for %s failed with error %s."), *CurrentMessage.TargetName, UTF8_TO_TCHAR(MLGetResultString(Result)));
+			TargetCreateMsg.TaskType = FTrackerMessage::TaskType::TargetCreateFailed;
+			OutgoingMessages.Enqueue(TargetCreateMsg);
+			return;
+		}
+
+		// [3] Cache all the static data for this target.
+		MLImageTrackerTargetStaticData Data;
+		FMemory::Memset(&Data, 0, sizeof(MLImageTrackerTargetStaticData));
+		Result = MLImageTrackerGetTargetStaticData(ImageTracker, Target, &Data);
+		if (Result != MLResult_Ok)
+		{
+			UE_LOG(LogMagicLeap, Error, TEXT("MLImageTrackerGetTargetStaticData failed with error %s."), UTF8_TO_TCHAR(MLGetResultString(Result)));
+			TargetCreateMsg.TaskType = FTrackerMessage::TaskType::TargetCreateFailed;
+			OutgoingMessages.Enqueue(TargetCreateMsg);
+			return;
+		}
+
+		TargetCreateMsg.TaskType = FTrackerMessage::TaskType::TargetCreateSucceeded;
+		TargetCreateMsg.Target = Target;
+		TargetCreateMsg.Data = Data;
+		OutgoingMessages.Enqueue(TargetCreateMsg);
+	}
+	else
+	{
+		UE_LOG(LogMagicLeap, Error, TEXT("Failed to get texture bulk data for image target %s. Image will not be tracked."), *CurrentMessage.TargetName);
+	}
+
+	for (int32 i = 0; i < MipPointers.Num(); ++i)
+	{
+		if (MipPointers[i] != nullptr)
+		{
+			FMemory::Free(MipPointers[i]);
+			MipPointers[i] = nullptr;
+		}
+	}
+
+	/*FTexture2DMipMap& Mip = CurrentMessage.TargetImageTexture->PlatformData->Mips[0];
+	const unsigned char* PixelData = static_cast<const unsigned char*>(Mip.BulkData.Lock(LOCK_READ_ONLY));
+	MLHandle Target = ML_INVALID_HANDLE;
+	MLResult Result = MLImageTrackerAddTargetFromArray(
+		ImageTracker,
+		&CurrentMessage.TargetSettings,
+		PixelData,
+		CurrentMessage.TargetImageTexture->GetSurfaceWidth(),
+		CurrentMessage.TargetImageTexture->GetSurfaceHeight(),
+		MLImageTrackerImageFormat_RGBA,
+		&Target);
+	Mip.BulkData.Unlock();
+
+	if (Result != MLResult_Ok)
+	{
+		UE_LOG(LogMagicLeap, Error, TEXT("MLImageTrackerAddTargetFromArray failed with error %d."), Result);
+		TargetCreateMsg.TaskType = FTrackerMessage::TaskType::TargetCreateFailed;
+		OutgoingMessages.Enqueue(TargetCreateMsg);
+		return;
+	}
+
+	// [3] Cache all the static data for this target.
+	MLImageTrackerTargetStaticData Data;
+	FMemory::Memset(&Data, 0, sizeof(MLImageTrackerTargetStaticData));
+	Result = MLImageTrackerGetTargetStaticData(ImageTracker, Target, &Data);
+	if (Result != MLResult_Ok)
+	{
+		UE_LOG(LogMagicLeap, Error, TEXT("MLImageTrackerGetTargetStaticData failed with error %d."), Result);
+		TargetCreateMsg.TaskType = FTrackerMessage::TaskType::TargetCreateFailed;
+		OutgoingMessages.Enqueue(TargetCreateMsg);
+		return;
+	}
+
+	TargetCreateMsg.TaskType = FTrackerMessage::TaskType::TargetCreateSucceeded;
+	TargetCreateMsg.Target = Target;
+	TargetCreateMsg.Data = Data;
+	OutgoingMessages.Enqueue(TargetCreateMsg);*/
+#endif //WITH_MLSDK
+}
+
+void FImageTrackerRunnable::UpdateTrackerSettings()
+{
+#if WITH_MLSDK
+	// If the tracker has already been created, update the setttings
+	// If the tracker has not been created, this cached Setting will be used whenever MLImageTrackerCreate() is called.
+	if (MLHandleIsValid(ImageTracker))
+	{
+		MLResult Result = MLImageTrackerUpdateSettings(ImageTracker, &Settings);
+		UE_CLOG(Result != MLResult_Ok, LogMagicLeap, Error, TEXT("MLImageTrackerUpdateSettings failed with error %s."), UTF8_TO_TCHAR(MLGetResultString(Result)));
+	}
+#endif //WITH_MLSDK
+}
\ No newline at end of file
diff --git a/Engine/Plugins/Lumin/MagicLeap/Source/MagicLeap/Private/ImageTrackerRunnable.h b/Engine/Plugins/Lumin/MagicLeap/Source/MagicLeap/Private/ImageTrackerRunnable.h
new file mode 100644
--- /dev/null
+++ b/Engine/Plugins/Lumin/MagicLeap/Source/MagicLeap/Private/ImageTrackerRunnable.h
@@ -0,0 +1,128 @@
+// %BANNER_BEGIN%
+// ---------------------------------------------------------------------
+// %COPYRIGHT_BEGIN%
+//
+// Copyright (c) 2017 Magic Leap, Inc. (COMPANY) All Rights Reserved.
+// Magic Leap, Inc. Confidential and Proprietary
+//
+// NOTICE: All information contained herein is, and remains the property
+// of COMPANY. The intellectual and technical concepts contained herein
+// are proprietary to COMPANY and may be covered by U.S. and Foreign
+// Patents, patents in process, and are protected by trade secret or
+// copyright law. Dissemination of this information or reproduction of
+// this material is strictly forbidden unless prior written permission is
+// obtained from COMPANY. Access to the source code contained herein is
+// hereby forbidden to anyone except current COMPANY employees, managers
+// or contractors who have executed Confidentiality and Non-disclosure
+// agreements explicitly covering such access.
+//
+// The copyright notice above does not evidence any actual or intended
+// publication or disclosure of this source code, which includes
+// information that is confidential and/or proprietary, and is a trade
+// secret, of COMPANY. ANY REPRODUCTION, MODIFICATION, DISTRIBUTION,
+// PUBLIC PERFORMANCE, OR PUBLIC DISPLAY OF OR THROUGH USE OF THIS
+// SOURCE CODE WITHOUT THE EXPRESS WRITTEN CONSENT OF COMPANY IS
+// STRICTLY PROHIBITED, AND IN VIOLATION OF APPLICABLE LAWS AND
+// INTERNATIONAL TREATIES. THE RECEIPT OR POSSESSION OF THIS SOURCE
+// CODE AND/OR RELATED INFORMATION DOES NOT CONVEY OR IMPLY ANY RIGHTS
+// TO REPRODUCE, DISCLOSE OR DISTRIBUTE ITS CONTENTS, OR TO MANUFACTURE,
+// USE, OR SELL ANYTHING THAT IT MAY DESCRIBE, IN WHOLE OR IN PART.
+//
+// %COPYRIGHT_END%
+// --------------------------------------------------------------------
+// %BANNER_END%
+
+#pragma once
+
+#include "IMagicLeapPlugin.h"
+#include "AppEventHandler.h"
+#include "HAL/Runnable.h"
+#include "HAL/RunnableThread.h"
+#include "Containers/Queue.h"
+#include "MagicLeapPluginUtil.h" // for ML_INCLUDES_START/END
+
+#if WITH_MLSDK
+ML_INCLUDES_START
+#include <ml_image_tracking.h>
+//#include <ml_perception.h>
+//#include <ml_snapshot.h>
+//#include <ml_head_tracking.h>
+//#include <ml_coordinate_frame_uid.h>
+ML_INCLUDES_END
+#endif //WITH_MLSDK
+
+struct FTrackerMessage
+{
+	enum TaskType
+	{
+		None,
+		Pause,
+		Resume,
+		SetEnabled,
+		SetMaxTargets,
+		TryCreateTarget,
+		TargetCreateFailed,
+		TargetCreateSucceeded,
+	};
+
+	TaskType TaskType;
+	class FImageTrackerImpl* Requester;
+#if WITH_MLSDK
+	MLHandle PrevTarget;
+	MLHandle Target;
+	MLImageTrackerTargetStaticData Data;
+	MLImageTrackerTargetSettings TargetSettings;
+#endif //WITH_MLSDK
+	FString TargetName;
+	class UTexture2D* TargetImageTexture;
+	bool bEnable;
+	uint32 MaxTargets;
+
+	FTrackerMessage()
+	: TaskType(TaskType::None)
+	, Requester(nullptr)
+#if WITH_MLSDK
+	, PrevTarget(ML_INVALID_HANDLE)
+	, Target(ML_INVALID_HANDLE)
+#endif //WITH_MLSDK
+	, TargetImageTexture(nullptr)
+	, bEnable(true)
+	, MaxTargets(1)
+	{}
+};
+
+class FImageTrackerRunnable : public FRunnable, public MagicLeap::IAppEventHandler
+{
+public:
+	FImageTrackerRunnable();
+	virtual ~FImageTrackerRunnable();
+
+#if WITH_MLSDK
+	MLHandle GetHandle() const;
+#endif //WITH_MLSDK
+	virtual uint32 Run() override;
+
+	TQueue<FTrackerMessage, EQueueMode::Spsc> IncomingMessages;
+	TQueue<FTrackerMessage, EQueueMode::Spsc> OutgoingMessages;
+
+private:
+
+	void OnAppPause() override;
+	void OnAppResume() override;
+	void TryPause();
+	void TryResume();
+	void OnAppShutDown() override;
+	void SetEnabled();
+	void SetMaxTargets();
+	void SetTarget();
+	void UpdateTrackerSettings();
+
+#if WITH_MLSDK
+	MLHandle ImageTracker;
+	MLImageTrackerSettings Settings;
+#endif //WITH_MLSDK
+	FRunnableThread* Thread;
+	FThreadSafeCounter StopTaskCounter;
+	FTrackerMessage CurrentMessage;
+	const float RetryCreateTrackerWaitTime;
+};
\ No newline at end of file
diff --git a/Engine/Plugins/Lumin/MagicLeap/Source/MagicLeap/Private/InAppPurchaseComponent.cpp b/Engine/Plugins/Lumin/MagicLeap/Source/MagicLeap/Private/InAppPurchaseComponent.cpp
new file mode 100644
--- /dev/null
+++ b/Engine/Plugins/Lumin/MagicLeap/Source/MagicLeap/Private/InAppPurchaseComponent.cpp
@@ -0,0 +1,450 @@
+// Fill out your copyright notice in the Description page of Project Settings.
+
+#include "InAppPurchaseComponent.h"
+#include "CoreMinimal.h"
+#include "Containers/Array.h"
+#include "Modules/ModuleManager.h"
+#include "AppEventHandler.h"
+#if PLATFORM_LUMIN
+#include "Lumin/LuminAffinity.h"
+#endif // PLATFORM_LUMIN
+
+#if WITH_MLSDK
+#include <ml_purchase.h>
+#endif //WITH_MLSDK
+
+DEFINE_LOG_CATEGORY(LogInAppPurchase);
+
+class FInAppPurchaseImpl : public MagicLeap::IAppEventHandler
+{
+public:
+	FInAppPurchaseImpl(UInAppPurchaseComponent* InOwner)
+	: CurrentRequests(0)
+#if WITH_MLSDK
+	, ItemsDetailsHandle(ML_INVALID_HANDLE)
+	, PurchaseHandle(ML_INVALID_HANDLE)
+	, PurchaseHistoryHandle(ML_INVALID_HANDLE)
+#endif //WITH_MLSDK
+	, Owner(InOwner)
+	{
+	}
+
+	virtual ~FInAppPurchaseImpl()
+	{
+#if WITH_MLSDK
+		if (MLHandleIsValid(ItemsDetailsHandle))
+		{
+			MLResult Result = MLPurchaseItemDetailsDestroy(ItemsDetailsHandle);
+			if (Result != MLResult_Ok)
+			{
+				Owner->Log(FString::Printf(TEXT("MLPurchaseItemDetailsDestroy failed with error %d"), Result));
+			}
+		}
+
+		if (MLHandleIsValid(PurchaseHandle))
+		{
+			MLResult Result = MLPurchaseDestroy(PurchaseHandle);
+			if (Result != MLResult_Ok)
+			{
+				Owner->Log(FString::Printf(TEXT("MLPurchaseDestroy failed with error %d"), Result));
+			}
+		}
+
+		if (MLHandleIsValid(PurchaseHistoryHandle))
+		{
+			MLResult Result = MLPurchaseHistoryQueryDestroy(PurchaseHistoryHandle);
+			if (Result != MLResult_Ok)
+			{
+				Owner->Log(FString::Printf(TEXT("MLPurchaseHistoryQueryDestroy failed with error %d"), Result));
+			}
+		}
+#endif //WITH_MLSDK
+	}
+
+	bool TryGetItemsDetails(const TArray<FString>& InItems)
+	{
+#if WITH_MLSDK
+		if ((CurrentRequests & RequestType::ItemsDetails) == 0)
+		{
+			MLResult Result = MLPurchaseItemDetailsCreate(&ItemsDetailsHandle);
+			if (Result != MLResult_Ok)
+			{
+				Owner->Log(FString::Printf(TEXT("MLItemDetailsCreate failed with error %d"), Result));
+			}
+			else
+			{
+				const char** ItemsIds = static_cast<const char**>(FMemory::Malloc(sizeof(char*) * InItems.Num()));
+				TArray<FString> ItemNames;
+				for (int32 Index = 0; Index < InItems.Num(); ++Index)
+				{
+					ItemsIds[Index] = static_cast<char*>(FMemory::Malloc(sizeof(char) * InItems[Index].Len() + 1));
+					FMemory::Memcpy(const_cast<char*>(ItemsIds[Index]), TCHAR_TO_UTF8(*InItems[Index]), InItems[Index].Len()+1);
+				}
+				MLPurchaseItemDetailsQuery ItemsDetailsQuery{ ItemsIds, static_cast<uint32_t>(InItems.Num()) };
+				Result = MLPurchaseItemDetailsGet(ItemsDetailsHandle, &ItemsDetailsQuery);
+				for (int32 Index = 0; Index < InItems.Num(); ++Index)
+				{
+					FMemory::Free(const_cast<char*>(ItemsIds[Index]));
+				}
+				FMemory::Free(ItemsIds);
+				if (Result != MLResult_Ok)
+				{
+					Owner->Log(FString::Printf(TEXT("MLPurchaseItemDetailsGet failed with error %d"), Result));
+					Result = MLPurchaseItemDetailsDestroy(ItemsDetailsHandle);
+					if (Result != MLResult_Ok)
+					{
+						Owner->Log(FString::Printf(TEXT("MLPurchaseItemDetailsDestroy failed with error %d"), Result));
+					}
+				}
+				else
+				{
+					CurrentRequests |= RequestType::ItemsDetails;
+					return true;
+				}
+			}
+		}
+#endif //WITH_MLSDK
+		return false;
+	}
+
+	bool TryGetPurchaseConfirmation(const FPurchaseItemDetails& ItemDetails)
+	{
+#if WITH_MLSDK
+		if ((CurrentRequests & RequestType::PurchaseConfirmation) == 0)
+		{
+			MLResult Result = MLPurchaseCreate(&PurchaseHandle);
+			if (Result != MLResult_Ok)
+			{
+				Owner->Log(FString::Printf(TEXT("MLPurchaseCreate failed with error %d"), Result));
+			}
+			else
+			{
+				Result = MLPurchaseSubmit(PurchaseHandle, ItemDetails.Token);
+				if (Result != MLResult_Ok)
+				{
+					Owner->Log(FString::Printf(TEXT("MLPurchaseSubmit failed with error %d"), Result));
+					Result = MLPurchaseDestroy(PurchaseHandle);
+					if (Result != MLResult_Ok)
+					{
+						Owner->Log(FString::Printf(TEXT("MLPurchaseDestroy failed with error %d"), Result));
+					}
+				}
+				else
+				{
+					CurrentRequests |= RequestType::PurchaseConfirmation;
+					return true;
+				}
+			}
+		}
+#endif //WITH_MLSDK
+		return false;
+	}
+
+	bool TryGetPurchaseHistory(int32 InNumPages)
+	{
+#if WITH_MLSDK
+		if ((CurrentRequests & RequestType::PurchaseHistory) == 0)
+		{
+			MLResult Result = MLPurchaseHistoryQueryCreate(&PurchaseHistoryHandle);
+			if (Result != MLResult_Ok)
+			{
+				Owner->Log(FString::Printf(TEXT("MLPurchaseHistoryQueryCreate failed with error %d"), Result));
+			}
+			else
+			{
+				Result = MLPurchaseHistoryQueryGetPage(PurchaseHistoryHandle, InNumPages);
+				if (Result != MLResult_Ok)
+				{
+					Owner->Log(FString::Printf(TEXT("MLPurchaseHistoryQueryGetPage failed with error %d"), Result));
+					Result = MLPurchaseHistoryQueryDestroy(PurchaseHistoryHandle);
+					if (Result != MLResult_Ok)
+					{
+						Owner->Log(FString::Printf(TEXT("MLPurchaseHistoryQueryDestroy failed with error %d"), Result));
+					}
+				}
+				else
+				{
+					CurrentRequests |= RequestType::PurchaseHistory;
+					return true;
+				}
+			}
+		}
+#endif //WITH_MLSDK
+		return false;
+	}
+
+	void Tick()
+	{
+#if WITH_MLSDK
+		MLCloudStatus CloudStatus;
+
+		if (CurrentRequests & ItemsDetails)
+		{
+			MLPurchaseItemDetailsResults MLItemsDetails;
+			MLResult Result = MLPurchaseItemDetailsGetResult(ItemsDetailsHandle, &MLItemsDetails, &CloudStatus);
+			if (Result != MLResult_Ok)
+			{
+				CurrentRequests &= (~ItemsDetails);
+				Owner->Log(FString::Printf(TEXT("MLPurchaseItemDetailsGetResult failed with error %d"), Result));
+				Owner->GetItemsDetailsFailure.Broadcast();
+			}
+			else if (CloudStatus == MLCloudStatus_Done)
+			{
+				CurrentRequests &= (~ItemsDetails);
+				TArray<FPurchaseItemDetails> UEItemsDetails;
+				MLToUE(MLItemsDetails, UEItemsDetails);
+				Owner->GetItemsDetailsSuccess.Broadcast(UEItemsDetails);
+			}
+
+			if (!(CurrentRequests & ItemsDetails))
+			{
+				Result = MLPurchaseItemDetailsDestroy(ItemsDetailsHandle);
+				if (Result != MLResult_Ok)
+				{
+					Owner->Log(FString::Printf(TEXT("MLPurchaseItemDetailsDestroy failed with error %d"), Result));
+				}
+			}
+		}
+
+		if (CurrentRequests & PurchaseConfirmation)
+		{
+			MLPurchaseConfirmation MLConfirmation;
+			MLResult Result = MLPurchaseGetResult(PurchaseHandle, &MLConfirmation, &CloudStatus);
+			if (Result != MLResult_Ok)
+			{
+				CurrentRequests &= (~PurchaseConfirmation);
+				Owner->Log(FString::Printf(TEXT("MLPurchaseGetResult failed with error %d"), Result));
+				Owner->PurchaseConfirmationFailure.Broadcast();
+			}
+			else if (CloudStatus == MLCloudStatus_Done)
+			{
+				CurrentRequests &= (~PurchaseConfirmation);
+				FPurchaseConfirmation UEConfirmation;
+				MLToUE(MLConfirmation, UEConfirmation);
+				Owner->PurchaseConfirmationSuccess.Broadcast(UEConfirmation);
+			}
+
+			if (!(CurrentRequests & PurchaseConfirmation))
+			{
+				Result = MLPurchaseDestroy(PurchaseHandle);
+				if (Result != MLResult_Ok)
+				{
+					Owner->Log(FString::Printf(TEXT("MLPurchaseDestroy failed with error %d"), Result));
+				}
+			}
+		}
+
+		if (CurrentRequests & PurchaseHistory)
+		{
+			MLPurchaseHistoryResult MLPurchaseHistory;
+			MLResult Result = MLPurchaseHistoryQueryGetPageResult(PurchaseHistoryHandle, &MLPurchaseHistory);
+			if (Result != MLResult_Ok)
+			{
+				CurrentRequests &= (~PurchaseHistory);
+				Owner->Log(FString::Printf(TEXT("MLPurchaseHistoryQueryGetPageResult failed with error %d"), Result));
+				Owner->GetPurchaseHistoryFailure.Broadcast();
+			}
+			else if (MLPurchaseHistory.status == MLCloudStatus_Done)
+			{
+				MLToUE(MLPurchaseHistory);
+				if (!MLPurchaseHistory.has_next_page)
+				{
+					CurrentRequests &= (~PurchaseHistory);
+					Owner->GetPurchaseHistorySuccess.Broadcast(CachedPurchaseHistory);
+					CachedPurchaseHistory.Empty();
+				}
+			}
+
+			if (!(CurrentRequests & PurchaseHistory))
+			{
+				Result = MLPurchaseHistoryQueryDestroy(PurchaseHistoryHandle);
+				if (Result != MLResult_Ok)
+				{
+					Owner->Log(FString::Printf(TEXT("MLPurchaseHistoryQueryDestroy failed with error %d"), Result));
+				}
+			}
+		}
+#endif //WITH_MLSDK
+	}
+
+	enum RequestType
+	{
+		None = 0,
+		ItemsDetails = 1,
+		PurchaseConfirmation = 2,
+		PurchaseHistory = 4,
+	};
+
+	uint32 CurrentRequests;
+#if WITH_MLSDK
+	MLHandle ItemsDetailsHandle;
+	MLHandle PurchaseHandle;
+	MLHandle PurchaseHistoryHandle;
+#endif //WITH_MLSDK
+	TArray<FPurchaseConfirmation> CachedPurchaseHistory;
+	UInAppPurchaseComponent* Owner;
+
+private:
+#if WITH_MLSDK
+	PurchaseType MLToUE(MLPurchaseType InMLPurchaseType)
+	{
+		switch (InMLPurchaseType)
+		{
+		case MLPurchaseType_Consumable: return PurchaseType::Consumable;
+		case MLPurchaseType_Nonconsumable: return PurchaseType::Nonconsumable;
+		}
+
+		return PurchaseType::Undefined;
+	}
+
+	void MLToUE(const MLPurchaseItemDetailsResults& MLItemsDetails, TArray<FPurchaseItemDetails>& UEItemsDetails)
+	{
+		UEItemsDetails.AddDefaulted(MLItemsDetails.count);
+		for (uint32_t Index = 0; Index < MLItemsDetails.count; ++Index)
+		{
+			const MLPurchaseItemDetailsResult& MLItemDetails = MLItemsDetails.item_details[Index];
+			FPurchaseItemDetails& UEItemDetails = UEItemsDetails[Index];
+			UEItemDetails.IAPID = MLItemDetails.iap_id;
+			UEItemDetails.Price = MLItemDetails.price;
+			UEItemDetails.Name = MLItemDetails.name;
+			UEItemDetails.Type = MLToUE(MLItemDetails.type);
+			UEItemDetails.Token = MLItemDetails.token;
+		}
+	}
+
+	void MLToUE(const MLPurchaseConfirmation& MLConfirmation, FPurchaseConfirmation& UEConfirmation)
+	{
+		UEConfirmation.OrderID = MLConfirmation.order_id;
+		UEConfirmation.PackageName = MLConfirmation.package_name;
+		UEConfirmation.PurchaseTime = MLConfirmation.purchase_time;
+		UEConfirmation.Signature = MLConfirmation.signature;
+		UEConfirmation.IAPID = MLConfirmation.iap_id;
+		UEConfirmation.Type = MLToUE(MLConfirmation.type);
+	}
+
+	void MLToUE(const MLPurchaseHistoryResult& MLPurchaseHistory)
+	{
+		int32 StartWriteIndex = CachedPurchaseHistory.Num();
+		CachedPurchaseHistory.AddDefaulted(MLPurchaseHistory.count);
+		for (size_t Index = 0; Index < MLPurchaseHistory.count; ++Index)
+		{
+			FPurchaseConfirmation& Confirmation = CachedPurchaseHistory[StartWriteIndex + Index];
+			MLToUE(MLPurchaseHistory.confirmations[Index], Confirmation);
+		}
+	}
+#endif //WITH_MLSDK
+};
+
+UInAppPurchaseComponent::UInAppPurchaseComponent()
+: Impl(nullptr)
+{
+	PrimaryComponentTick.TickGroup = TG_PrePhysics;
+	PrimaryComponentTick.bStartWithTickEnabled = true;
+	PrimaryComponentTick.bCanEverTick = true;
+}
+
+UInAppPurchaseComponent::~UInAppPurchaseComponent()
+{
+	delete Impl;
+	Impl = nullptr;
+}
+
+void UInAppPurchaseComponent::BeginPlay()
+{
+	Super::BeginPlay();
+	Impl = new FInAppPurchaseImpl(this);
+}
+
+void UInAppPurchaseComponent::TickComponent(float DeltaTime, enum ELevelTick TickType, FActorComponentTickFunction *ThisTickFunction)
+{
+	Super::TickComponent(DeltaTime, TickType, ThisTickFunction);
+	Impl->Tick();
+}
+
+UInAppPurchaseComponent::FInAppPurchaseLogMessage& UInAppPurchaseComponent::OnInAppPurchaseLogMessage()
+{
+  return InAppPurchaseLogMessage;
+}
+
+UInAppPurchaseComponent::FGetItemsDetailsSuccess& UInAppPurchaseComponent::OnGetItemsDetailsSuccess()
+{
+  return GetItemsDetailsSuccess;
+}
+
+UInAppPurchaseComponent::FGetItemsDetailsFailure& UInAppPurchaseComponent::OnGetItemsDetailsFailure()
+{
+  return GetItemsDetailsFailure;
+}
+
+UInAppPurchaseComponent::FPurchaseConfirmationSuccess& UInAppPurchaseComponent::OnPurchaseConfirmationSuccess()
+{
+  return PurchaseConfirmationSuccess;
+}
+
+UInAppPurchaseComponent::FPurchaseConfirmationFailure& UInAppPurchaseComponent::OnPurchaseConfirmationFailure()
+{
+  return PurchaseConfirmationFailure;
+}
+
+UInAppPurchaseComponent::FGetPurchaseHistorySuccess& UInAppPurchaseComponent::OnGetPurchaseHistorySuccess()
+{
+  return GetPurchaseHistorySuccess;
+}
+
+UInAppPurchaseComponent::FGetPurchaseHistoryFailure& UInAppPurchaseComponent::OnGetPurchaseHistoryFailure()
+{
+  return GetPurchaseHistoryFailure;
+}
+
+bool UInAppPurchaseComponent::TryGetItemsDetailsAsync(const TArray<FString>& ItemIDs)
+{
+  if (Impl->TryGetItemsDetails(ItemIDs))
+  {
+    return true;
+  }
+
+  Log("Items details query already in progress!", false);
+  return false;
+}
+
+bool UInAppPurchaseComponent::TryPurchaseItemAsync(const FPurchaseItemDetails& ItemDetails)
+{
+  if (Impl->TryGetPurchaseConfirmation(ItemDetails))
+  {
+    return true;
+  }
+
+  Log("Item purchase already in progress!", false);
+  return false;
+}
+
+bool UInAppPurchaseComponent::TryGetPurchaseHistoryAsync(int32 InNumPages)
+{
+	if (InNumPages <= 0)
+	{
+		Log(FString::Printf(TEXT("TryGetPurchaseHistoryAsync failed due to invalid number of pages (%d)"), InNumPages));
+		return false;
+	}
+
+	if (Impl->TryGetPurchaseHistory(InNumPages))
+	{
+		return true;
+	}
+
+	Log("Purchase history query already in progress!", false);
+	return false;
+}
+
+void UInAppPurchaseComponent::Log(const FString& LogMessage, bool bError/* = true */)
+{
+	if (bError)
+	{
+		UE_LOG(LogInAppPurchase, Error, TEXT("%s"), *LogMessage);
+	}
+	else
+	{
+		UE_LOG(LogInAppPurchase, Display, TEXT("%s"), *LogMessage);
+	}
+
+	InAppPurchaseLogMessage.Broadcast(LogMessage);
+}
diff --git a/Engine/Plugins/Lumin/MagicLeap/Source/MagicLeap/Private/LightingTrackingComponent.cpp b/Engine/Plugins/Lumin/MagicLeap/Source/MagicLeap/Private/LightingTrackingComponent.cpp
new file mode 100644
--- /dev/null
+++ b/Engine/Plugins/Lumin/MagicLeap/Source/MagicLeap/Private/LightingTrackingComponent.cpp
@@ -0,0 +1,375 @@
+// %BANNER_BEGIN%
+// ---------------------------------------------------------------------
+// %COPYRIGHT_BEGIN%
+//
+// Copyright (c) 2017 Magic Leap, Inc. (COMPANY) All Rights Reserved.
+// Magic Leap, Inc. Confidential and Proprietary
+//
+// NOTICE: All information contained herein is, and remains the property
+// of COMPANY. The intellectual and technical concepts contained herein
+// are proprietary to COMPANY and may be covered by U.S. and Foreign
+// Patents, patents in process, and are protected by trade secret or
+// copyright law. Dissemination of this information or reproduction of
+// this material is strictly forbidden unless prior written permission is
+// obtained from COMPANY. Access to the source code contained herein is
+// hereby forbidden to anyone except current COMPANY employees, managers
+// or contractors who have executed Confidentiality and Non-disclosure
+// agreements explicitly covering such access.
+//
+// The copyright notice above does not evidence any actual or intended
+// publication or disclosure of this source code, which includes
+// information that is confidential and/or proprietary, and is a trade
+// secret, of COMPANY. ANY REPRODUCTION, MODIFICATION, DISTRIBUTION,
+// PUBLIC PERFORMANCE, OR PUBLIC DISPLAY OF OR THROUGH USE OF THIS
+// SOURCE CODE WITHOUT THE EXPRESS WRITTEN CONSENT OF COMPANY IS
+// STRICTLY PROHIBITED, AND IN VIOLATION OF APPLICABLE LAWS AND
+// INTERNATIONAL TREATIES. THE RECEIPT OR POSSESSION OF THIS SOURCE
+// CODE AND/OR RELATED INFORMATION DOES NOT CONVEY OR IMPLY ANY RIGHTS
+// TO REPRODUCE, DISCLOSE OR DISTRIBUTE ITS CONTENTS, OR TO MANUFACTURE,
+// USE, OR SELL ANYTHING THAT IT MAY DESCRIBE, IN WHOLE OR IN PART.
+//
+// %COPYRIGHT_END%
+// --------------------------------------------------------------------
+// %BANNER_END%
+#include "LightingTrackingComponent.h"
+#include "IImageWrapper.h"
+#include "IImageWrapperModule.h"
+#include "Modules/ModuleManager.h"
+#include "Engine/TextureCube.h"
+#include "RenderUtils.h"
+#include "EngineUtils.h"
+#include "GameFramework/Actor.h"
+#include "Components/PostProcessComponent.h"
+#if WITH_MLSDK
+#include "ml_lighting_tracking.h"
+#endif //WITH_MLSDK
+#include <limits>
+
+DEFINE_LOG_CATEGORY(LogLightingTracking);
+
+#define MAX_NITS 15.0f
+
+class LightingTrackingImpl
+{
+public:
+	LightingTrackingImpl(ULightingTrackingComponent* InOwner)
+		: Owner(InOwner)
+#if WITH_MLSDK
+		, Tracker(ML_INVALID_HANDLE)
+#endif //WITH_MLSDK
+		, PostProcessor(nullptr)
+		, AmbientCubeMap(nullptr)
+		, LastAmbientIntensityTimeStamp(0)
+		, LastAmbientCubeMapTimeStamp(0)
+	{
+#if WITH_MLSDK
+		MLResult Result = MLLightingTrackingCreate(&Tracker);
+		if (Result != MLResult_Ok)
+		{
+			UE_LOG(LogLightingTracking, Error, TEXT("MLLightingTrackingCreate failed with error %d"), Result);
+		}
+#endif //WITH_MLSDK
+
+		for (TActorIterator<AActor> ActorItr(Owner->GetWorld()); ActorItr; ++ActorItr)
+		{
+			PostProcessor = (*ActorItr)->FindComponentByClass<UPostProcessComponent>();
+			if (PostProcessor)
+			{
+				break;
+			}
+		}
+
+		if (!PostProcessor)
+		{
+			PostProcessor = NewObject<UPostProcessComponent>(Owner->GetOwner());
+			PostProcessor->RegisterComponent();
+		}
+
+		PostProcessor->bUnbound = true;
+		PostProcessor->Settings.bOverride_AmbientCubemapIntensity = Owner->UseGlobalAmbience;
+		PostProcessor->Settings.bOverride_WhiteTemp = Owner->UseColorTemp;
+
+		/*const int32 MinDim = static_cast<int32>(FMath::Min(MLLightingTracking_GridWidth, MLLightingTracking_GridHeight));
+		const int32 NumCubeMapTextures = 6;
+
+		AmbientCubeMap = NewObject<UTextureCube>(
+			GetTransientPackage(),
+			NAME_None,
+			RF_Transient
+		);
+		AmbientCubeMap->AddToRoot();
+		AmbientCubeMap->PlatformData = new FTexturePlatformData();
+		AmbientCubeMap->PlatformData->SizeX = MinDim * NumCubeMapTextures;
+		AmbientCubeMap->PlatformData->SizeY = MinDim;
+		AmbientCubeMap->PlatformData->PixelFormat = EPixelFormat::PF_FloatRGBA;
+
+		// Allocate first mipmap.
+		int32 NumBlocksX = (MinDim * NumCubeMapTextures) / GPixelFormats[EPixelFormat::PF_FloatRGBA].BlockSizeX;
+		int32 NumBlocksY = MinDim / GPixelFormats[EPixelFormat::PF_FloatRGBA].BlockSizeY;
+		FTexture2DMipMap* Mip = new FTexture2DMipMap();// (AmbientCubeMap->PlatformData->Mips) FTexture2DMipMap();
+		Mip->SizeX = MinDim * 6;
+		Mip->SizeY = MinDim;
+		Mip->BulkData.Lock(LOCK_READ_WRITE);
+		uint16* Data = static_cast<uint16*>(Mip->BulkData.Realloc(NumBlocksX * NumBlocksY * GPixelFormats[EPixelFormat::PF_FloatRGBA].BlockBytes));
+		for (int32 i = 0; i < Mip->BulkData.GetBulkDataSize() / 2; i+=2, ++Data)
+		{
+			*((FFloat16*)Data) = 1.0f;
+		}
+		Mip->BulkData.Unlock();
+		AmbientCubeMap->PlatformData->Mips.Add(Mip);
+		AmbientCubeMap->UpdateResource();
+
+		PostProcessor->Settings.AmbientCubemap = AmbientCubeMap;*/
+	}
+
+	void RefreshGlobalAmbience()
+	{
+#if WITH_MLSDK
+		PostProcessor->Settings.bOverride_AmbientCubemapIntensity = Owner->UseGlobalAmbience;
+
+		if (!Owner->UseGlobalAmbience)
+		{
+			return;
+		}
+
+		MLLightingTrackingAmbientGlobalState AmbientGlobalState;
+		MLResult Result = MLLightingTrackingGetAmbientGlobalState(Tracker, &AmbientGlobalState);
+		if (Result == MLResult_Ok)
+		{
+			if (AmbientGlobalState.timestamp_ns > LastAmbientIntensityTimeStamp)
+			{
+				LastAmbientIntensityTimeStamp = AmbientGlobalState.timestamp_ns;
+				uint16 LuminanceSum = 0;
+				for (uint16 Index = 0; Index < (uint16)MLLightingTrackingCamera_Count; ++Index)
+				{
+					LuminanceSum += AmbientGlobalState.als_global[Index];
+				}
+				checkf((uint16)MLLightingTrackingCamera_Count > 0, TEXT("Invalid value for MLLightingTrackingCamera_Count!"));
+				PostProcessor->Settings.AmbientCubemapIntensity = (float)(LuminanceSum / MLLightingTrackingCamera_Count) / MAX_NITS;// (float)std::numeric_limits<uint16>::max();
+			}
+		}
+		else
+		{
+			UE_LOG(LogLightingTracking, Error, TEXT("MLLightingTrackingGetAmbientGlobalState failed with error %d"), Result);
+		}
+#endif //WITH_MLSDK
+	}
+
+	void RefreshColorTemperature()
+	{
+#if WITH_MLSDK
+		PostProcessor->Settings.bOverride_WhiteTemp = Owner->UseColorTemp;
+
+		if (!Owner->UseColorTemp)
+		{
+			return;
+		}
+
+		MLLightingTrackingColorTemperatureState ColorTemperatureState;
+		MLResult Result = MLLightingTrackingGetColorTemperatureState(Tracker, &ColorTemperatureState);
+		if (Result == MLResult_Ok)
+		{
+			PostProcessor->Settings.WhiteTemp = (float)ColorTemperatureState.color_temp;
+		}
+		else
+		{
+			UE_LOG(LogLightingTracking, Error, TEXT("MLLightingTrackingGetColorTemperatureState failed with error %d"), Result);
+		}
+#endif //WITH_MLSDK
+	}
+
+	void RefreshAmbientCubeMap()
+	{
+		/*if (!AmbientCubeMap)
+		{
+			UE_LOG(LogLightingTracking, Error, TEXT("AmbientCubeMap has not been set!"));
+		}
+
+		MLLightingTrackingAmbientGridState AmbientGridState;
+		MLResult Result = MLLightingTrackingGetAmbientGridState(Tracker, &AmbientGridState);
+		if (Result == MLResult_Ok)
+		{
+			if (AmbientGridState.timestamp_ns > LastAmbientCubeMapTimeStamp)
+			{
+				LastAmbientCubeMapTimeStamp = AmbientGridState.timestamp_ns;
+				ENQUEUE_UNIQUE_RENDER_COMMAND_THREEPARAMETER
+				(
+					RefreshAmbientMapCommand,
+					UTextureCube*,
+					_AmbientCubeMap,
+					AmbientCubeMap,
+					MLLightingTrackingAmbientGridState,
+					_AmbientGridState,
+					AmbientGridState,
+					LightingTrackingImpl*,
+					_LightingTrackingImpl,
+					this,
+					{
+						const uint32 TextureSizeInBytes = MLLightingTracking_GridHeight * MLLightingTracking_GridHeight * sizeof(FFloat16Color);
+						FTexture2DMipMap& Mip = _AmbientCubeMap->PlatformData->Mips[0];
+						void* Data = Mip.BulkData.Lock(LOCK_READ_WRITE);
+						_LightingTrackingImpl->CalcRight(_AmbientGridState.als_grid[MLLightingTrackingCamera_FarRight], static_cast<FFloat16Color*>(Data));
+						Data += TextureSizeInBytes;
+						_LightingTrackingImpl->CalcLeft(_AmbientGridState.als_grid[MLLightingTrackingCamera_FarLeft], static_cast<FFloat16Color*>(Data));
+						Data += TextureSizeInBytes;
+						_LightingTrackingImpl->CalcFront(_AmbientGridState.als_grid[MLLightingTrackingCamera_Left], static_cast<FFloat16Color*>(Data));
+						Data += TextureSizeInBytes;
+						_LightingTrackingImpl->CalcBack(_AmbientGridState.als_grid[MLLightingTrackingCamera_Left], static_cast<FFloat16Color*>(Data));
+						Data += TextureSizeInBytes;
+						_LightingTrackingImpl->CalcUp(_AmbientGridState.als_grid[MLLightingTrackingCamera_Left], static_cast<FFloat16Color*>(Data));
+						Data += TextureSizeInBytes;
+						_LightingTrackingImpl->CalcDown(_AmbientGridState.als_grid[MLLightingTrackingCamera_Left], static_cast<FFloat16Color*>(Data));
+						Mip.BulkData.Unlock();
+					}
+				);
+
+				AmbientCubeMap->UpdateResource();
+			}
+		}
+		else
+		{
+			UE_LOG(LogLightingTracking, Error, TEXT("MLLightingTrackingGetColorTemperatureState failed with error %d"), Result);
+		}*/
+	}
+#if WITH_MLSDK
+	void CalcRight(uint16_t InLumins[MLLightingTracking_GridHeight][MLLightingTracking_GridWidth], FFloat16Color* OutColors)
+	{
+		uint32 ColorCount = 0;
+		for (int32 YIndex = 0; YIndex < MLLightingTracking_GridHeight; ++YIndex)
+		{
+			for (int32 XIndex = MLLightingTracking_GridHeight; XIndex > 0; --XIndex)
+			{
+				const uint16 Lumins = InLumins[YIndex][XIndex];
+				FFloat16Color* Color = &OutColors[ColorCount++];
+				Color->R = Lumins;
+				Color->G = Lumins;
+				Color->B = Lumins;
+				Color->A = 1.0f;
+			}
+		}
+	}
+
+	void CalcLeft(uint16_t InLumins[MLLightingTracking_GridHeight][MLLightingTracking_GridWidth], FFloat16Color* OutColors)
+	{
+		uint32 ColorCount = 0;
+		for (int32 XIndex = 1; XIndex < MLLightingTracking_GridHeight; ++XIndex)
+		{
+			for (int32 YIndex = MLLightingTracking_GridHeight - 1; YIndex > -1; --YIndex)
+			{
+				const uint16 Lumins = InLumins[YIndex][XIndex];
+				FFloat16Color* Color = &OutColors[ColorCount++];
+				Color->R = Lumins;
+				Color->G = Lumins;
+				Color->B = Lumins;
+				Color->A = 1.0f;
+			}
+		}
+	}
+
+	void CalcFront(uint16 InLumins[MLLightingTracking_GridHeight][MLLightingTracking_GridWidth], FFloat16Color* OutColors)
+	{
+		uint32 ColorCount = 0;
+		for (int32 YIndex = MLLightingTracking_GridHeight - 1; YIndex > -1; --YIndex)
+		{
+			for (int32 XIndex = MLLightingTracking_GridWidth - 2; XIndex > 0; ++XIndex)
+			{
+				const uint16 Lumins = InLumins[YIndex][XIndex];
+				FFloat16Color* Color = &OutColors[ColorCount++];
+				Color->R = Lumins;
+				Color->G = Lumins;
+				Color->B = Lumins;
+				Color->A = 1.0f;
+			}
+		}
+	}
+
+	void CalcBack(uint16 InLumins[MLLightingTracking_GridHeight][MLLightingTracking_GridWidth], FFloat16Color* OutColors)
+	{
+		(void)InLumins;
+		for (uint32 Index = 0; Index < MLLightingTracking_GridHeight * MLLightingTracking_GridHeight; ++Index)
+		{
+			FFloat16Color* Color = &OutColors[Index];
+			Color->R = 0.5f;
+			Color->G = 0.5f;
+			Color->B = 0.5f;
+			Color->A = 1.0f;
+		}
+	}
+
+	void CalcUp(uint16 InLumins[MLLightingTracking_GridHeight][MLLightingTracking_GridWidth], FFloat16Color* OutColors)
+	{
+		(void)InLumins;
+		for (uint32 Index = 0; Index < MLLightingTracking_GridHeight * MLLightingTracking_GridHeight; ++Index)
+		{
+			FFloat16Color* Color = &OutColors[Index];
+			Color->R = 0.5f;
+			Color->G = 0.5f;
+			Color->B = 0.5f;
+			Color->A = 1.0f;
+		}
+	}
+
+	void CalcDown(uint16 InLumins[MLLightingTracking_GridHeight][MLLightingTracking_GridWidth], FFloat16Color* OutColors)
+	{
+		(void)InLumins;
+		for (uint32 Index = 0; Index < MLLightingTracking_GridHeight * MLLightingTracking_GridHeight; ++Index)
+		{
+			FFloat16Color* Color = &OutColors[Index];
+			Color->R = 0.5f;
+			Color->G = 0.5f;
+			Color->B = 0.5f;
+			Color->A = 1.0f;
+		}
+	}
+
+	MLHandle Tracker;
+#endif //WITH_MLSDK
+	ULightingTrackingComponent* Owner;
+	UPostProcessComponent* PostProcessor;
+	UTextureCube* AmbientCubeMap;
+	uint64 LastAmbientIntensityTimeStamp;
+	uint64 LastAmbientCubeMapTimeStamp;
+};
+
+ULightingTrackingComponent::ULightingTrackingComponent()
+	: Impl(nullptr)
+	, UseGlobalAmbience(false)
+	, UseColorTemp(false)
+	//, UseDynamicAmbientCubeMap(false)
+{
+	PrimaryComponentTick.TickGroup = TG_PrePhysics;
+	PrimaryComponentTick.bStartWithTickEnabled = true;
+	PrimaryComponentTick.bCanEverTick = true;
+}
+
+void ULightingTrackingComponent::BeginPlay()
+{
+	UActorComponent::BeginPlay();
+	Impl = new LightingTrackingImpl(this);
+}
+
+void ULightingTrackingComponent::EndPlay(const EEndPlayReason::Type EndPlayReason)
+{
+	UActorComponent::EndPlay(EndPlayReason);
+#if WITH_MLSDK
+	if (MLHandleIsValid(Impl->Tracker))
+	{
+		MLResult Result = MLLightingTrackingDestroy(Impl->Tracker);
+		if (Result != MLResult_Ok)
+		{
+			UE_LOG(LogLightingTracking, Error, TEXT("MLLightingTrackingDestroy failed with error %d"), Result);
+		}
+	}
+#endif //WITH_MLSDK
+	delete Impl;
+	Impl = nullptr;
+}
+
+void ULightingTrackingComponent::TickComponent(float DeltaTime, enum ELevelTick TickType, FActorComponentTickFunction *ThisTickFunction)
+{
+	Impl->RefreshGlobalAmbience();
+
+	Impl->RefreshColorTemperature();
+
+	//Impl->RefreshAmbientCubeMap();
+}
\ No newline at end of file
diff --git a/Engine/Plugins/Lumin/MagicLeap/Source/MagicLeap/Private/MagicLeapARPinComponent.cpp b/Engine/Plugins/Lumin/MagicLeap/Source/MagicLeap/Private/MagicLeapARPinComponent.cpp
new file mode 100644
--- /dev/null
+++ b/Engine/Plugins/Lumin/MagicLeap/Source/MagicLeap/Private/MagicLeapARPinComponent.cpp
@@ -0,0 +1,583 @@
+// %BANNER_BEGIN%
+// ---------------------------------------------------------------------
+// %COPYRIGHT_BEGIN%
+//
+// Copyright (c) 2017 Magic Leap, Inc. (COMPANY) All Rights Reserved.
+// Magic Leap, Inc. Confidential and Proprietary
+//
+// NOTICE: All information contained herein is, and remains the property
+// of COMPANY. The intellectual and technical concepts contained herein
+// are proprietary to COMPANY and may be covered by U.S. and Foreign
+// Patents, patents in process, and are protected by trade secret or
+// copyright law. Dissemination of this information or reproduction of
+// this material is strictly forbidden unless prior written permission is
+// obtained from COMPANY. Access to the source code contained herein is
+// hereby forbidden to anyone except current COMPANY employees, managers
+// or contractors who have executed Confidentiality and Non-disclosure
+// agreements explicitly covering such access.
+//
+// The copyright notice above does not evidence any actual or intended
+// publication or disclosure of this source code, which includes
+// information that is confidential and/or proprietary, and is a trade
+// secret, of COMPANY. ANY REPRODUCTION, MODIFICATION, DISTRIBUTION,
+// PUBLIC PERFORMANCE, OR PUBLIC DISPLAY OF OR THROUGH USE OF THIS
+// SOURCE CODE WITHOUT THE EXPRESS WRITTEN CONSENT OF COMPANY IS
+// STRICTLY PROHIBITED, AND IN VIOLATION OF APPLICABLE LAWS AND
+// INTERNATIONAL TREATIES. THE RECEIPT OR POSSESSION OF THIS SOURCE
+// CODE AND/OR RELATED INFORMATION DOES NOT CONVEY OR IMPLY ANY RIGHTS
+// TO REPRODUCE, DISCLOSE OR DISTRIBUTE ITS CONTENTS, OR TO MANUFACTURE,
+// USE, OR SELL ANYTHING THAT IT MAY DESCRIBE, IN WHOLE OR IN PART.
+//
+// %COPYRIGHT_END%
+// --------------------------------------------------------------------
+// %BANNER_END%
+
+#include "MagicLeapARPinComponent.h"
+#include "IMagicLeapPlugin.h"
+#include "MagicLeapHMD.h"
+#include "AppFramework.h"
+#include "AppEventHandler.h"
+#include "MagicLeapMath.h"
+#include "Engine/Engine.h"
+#include "GameFramework/Actor.h"
+#include "HAL/Runnable.h"
+#include "HAL/RunnableThread.h"
+#include "Kismet/KismetMathLibrary.h"
+#include "HeadMountedDisplayFunctionLibrary.h"
+#include "MagicLeapSecureStorage.h"
+#if PLATFORM_LUMIN
+#include "Lumin/LuminAffinity.h"
+#endif // PLATFORM_LUMIN
+
+#if WITH_MLSDK
+#include "ml_passable_world.h"
+#include "ml_persistent_coordinate_frames.h"
+// #include "ml_secure_storage.h"
+#endif //WITH_MLSDK
+
+#if WITH_MLSDK
+// TODO: Don't rely on the size being same. Use an object abstraction similar to UARPin.
+static_assert(sizeof(FGuid) == sizeof(MLCoordinateFrameUID), "Size of FGuid should be same as MLCoordinateFrameUID. TODO: Don't rely on the size being same. Use an object abstraction similar to UARPin.");
+#endif //WITH_MLSDK
+
+struct FPersistentData
+{
+#if WITH_MLSDK
+	MLCoordinateFrameUID PinnedCFUID;
+#endif //WITH_MLSDK
+	FTransform ComponentRelativeToCFUID;
+
+	FPersistentData()
+	{}
+
+#if WITH_MLSDK
+	FPersistentData(const MLCoordinateFrameUID& InPinnedCFUID, const FTransform& InComponentRelativeToCFUID)
+	: ComponentRelativeToCFUID(InComponentRelativeToCFUID)
+	{
+		FMemory::Memcpy(PinnedCFUID, InPinnedCFUID);
+	}
+#endif //WITH_MLSDK
+
+	FPersistentData(const FGuid& InPinnedCFUID, const FTransform& InComponentRelativeToCFUID)
+	: ComponentRelativeToCFUID(InComponentRelativeToCFUID)
+	{
+#if WITH_MLSDK
+		FMemory::Memcpy(&PinnedCFUID, &InPinnedCFUID, sizeof(MLCoordinateFrameUID));
+#endif //WITH_MLSDK
+	}
+};
+
+#if WITH_MLSDK
+EPassableWorldError MLToUnrealPassableWorldError(MLResult result)
+{
+	switch (result)
+	{
+		case MLResult_Ok: return EPassableWorldError::None;
+		case MLPassableWorldResult_LowMapQuality: return EPassableWorldError::LowMapQuality;
+		case MLPassableWorldResult_UnableToLocalize: return EPassableWorldError::UnableToLocalize;
+		case MLPassableWorldResult_ServerUnavailable: return EPassableWorldError::Unavailable;
+		case MLResult_PrivilegeDenied: return EPassableWorldError::PrivilegeDenied;
+		case MLResult_InvalidParam: return EPassableWorldError::InvalidParam;
+		case MLResult_UnspecifiedFailure: return EPassableWorldError::UnspecifiedFailure;
+	}
+	return EPassableWorldError::UnspecifiedFailure;
+}
+#endif //WITH_MLSDK
+
+class FMagicLeapARPinInterface : public FRunnable, public MagicLeap::IAppEventHandler
+{
+public:
+	static TWeakPtr<FMagicLeapARPinInterface, ESPMode::ThreadSafe> Get()
+	{
+		if (!Instance.IsValid())
+		{
+			Instance = MakeShareable(new FMagicLeapARPinInterface());
+		}
+
+		return Instance;
+	}
+
+	virtual ~FMagicLeapARPinInterface()
+	{
+		StopTaskCounter.Increment();
+
+		if (nullptr != Thread)
+		{
+			Thread->WaitForCompletion();
+			delete Thread;
+			Thread = nullptr;
+		}
+	}
+
+#if WITH_MLSDK
+	MLHandle GetHandle() const
+	{
+		return Tracker;
+	}
+#endif //WITH_MLSDK
+
+#if WITH_MLSDK
+	bool IsCoordinateFrameValid(const MLCoordinateFrameUID& InCoordinateFrame) const
+	{
+		FScopeLock Lock(&Mutex);
+		return (AllCoordinateFrames.FindByPredicate([&](const MLCoordinateFrameUID& ArrayElem) {
+			return (ArrayElem.data[0] == InCoordinateFrame.data[0] && ArrayElem.data[1] == InCoordinateFrame.data[1]);
+		}) != nullptr);
+	}
+
+	EPassableWorldError GetClosestARPin(const FVector& SearchPoint, FGuid& PinID)
+	{
+		EPassableWorldError ErrorReturn = EPassableWorldError::Unavailable;
+		if (GetPrivilegeStatus(MLPrivilegeID_PwFoundObjRead) == MagicLeap::EPrivilegeState::Granted)
+		{
+			if (IMagicLeapPlugin::Get().IsMagicLeapHMDValid())
+			{
+				const FAppFramework& AppFramework = static_cast<FMagicLeapHMD*>(GEngine->XRSystem->GetHMDDevice())->GetAppFrameworkConst();
+				const float WorldToMetersScale = AppFramework.GetWorldToMetersScale();
+				FTransform PoseInverse = UHeadMountedDisplayFunctionLibrary::GetTrackingToWorldTransform(GWorld).Inverse();
+				MLVec3f Target = MagicLeap::ToMLVector(PoseInverse.TransformPosition(SearchPoint), WorldToMetersScale);
+				FScopeLock Lock(&Mutex);
+				MLResult Result = MLPersistentCoordinateFrameGetClosest(GetHandle(), &Target, reinterpret_cast<MLCoordinateFrameUID*>(&PinID));
+				UE_CLOG(MLResult_Ok != Result, LogMagicLeap, Error, TEXT("MLPersistentCoordinateFrameGetClosest failed with error %s"), UTF8_TO_TCHAR(MLPersistentCoordinateFrameGetResultString(Result)));
+				ErrorReturn = MLToUnrealPassableWorldError(Result);
+			}
+		}
+		else
+		{
+			UE_LOG(LogMagicLeap, Error, TEXT("GetClosestARPin failed due to lack of privilege!"));
+		}
+
+		return ErrorReturn;
+	}
+
+	MLResult GetAvailableARPins(int32 NumRequested, TArray<FGuid>& CoordinateFrames) const
+	{
+		FScopeLock Lock(&Mutex);
+		NumRequested = (NumRequested <= AllCoordinateFrames.Num()) ? NumRequested : AllCoordinateFrames.Num();
+		CoordinateFrames.Reset();
+		CoordinateFrames.AddUninitialized(NumRequested);
+		FMemory::Memcpy(CoordinateFrames.GetData(), AllCoordinateFrames.GetData(), sizeof(MLCoordinateFrameUID) * NumRequested);
+		return LastQueryResult;
+	}
+
+	EPassableWorldError GetNumAvailableARPins(int32& Count)
+	{
+		MLResult Result = MLResult_PrivilegeDenied;
+		if (GetPrivilegeStatus(MLPrivilegeID_PwFoundObjRead) == MagicLeap::EPrivilegeState::Granted)
+		{
+			uint32 NumPersistentFrames = 0;
+			FScopeLock Lock(&Mutex);
+			Result = MLPersistentCoordinateFrameGetCount(GetHandle(), &NumPersistentFrames);
+			if (MLResult_Ok == Result)
+			{
+				Count = static_cast<int32>(NumPersistentFrames);
+			}
+			else
+			{
+				UE_LOG(LogMagicLeap, Error, TEXT("MLPersistentCoordinateFrameGetCount failed with error %s"), UTF8_TO_TCHAR(MLPersistentCoordinateFrameGetResultString(Result)));
+				Count = 0;
+			}
+		}
+		else
+		{
+			UE_LOG(LogMagicLeap, Error, TEXT("GetNumAvailableARPins failed due to lack of privilege!"));
+		}
+
+		return MLToUnrealPassableWorldError(Result);
+	}
+#endif //WITH_MLSDK
+
+	float GetThreadSleepTime() const
+	{
+		FScopeLock Lock(&Mutex);
+		return ThreadSleepTime;
+	}
+
+	void SetThreadSleepTime(float QueryInterval)
+	{
+		FScopeLock Lock(&Mutex);
+		ThreadSleepTime = QueryInterval;
+	}
+
+	virtual uint32 Run() override
+	{
+		// TODO: All this is temporary, until we get a proper return value from MLSnapshotGetTransform() indicating that the CFUID was invalid or not found in current map.
+		while (StopTaskCounter.GetValue() == 0)
+		{
+#if WITH_MLSDK
+			if (!MLHandleIsValid(Tracker))
+			{
+				if (GetPrivilegeStatus(MLPrivilegeID_PwFoundObjRead) == MagicLeap::EPrivilegeState::Granted)
+				{
+					// TODO: add retires like in Image tracker if error is LowMapQuality, UnableToLocalize, ServerUnavailable or PrivilegeDenied.
+					// Retrying for PrivilegeDenied would only make sense when MLPrivilege runtime api is functional.
+					MLResult Result = MLPersistentCoordinateFrameTrackerCreate(&Tracker);
+					UE_CLOG(Result != MLResult_Ok, LogMagicLeap, Error, TEXT("Failed to create persistent coordinate frame tracker with error %s."), UTF8_TO_TCHAR(MLPersistentCoordinateFrameGetResultString(Result)));
+				}
+				else
+				{
+					UE_LOG(LogMagicLeap, Error, TEXT("Failed to initialize persistent coordinate frame tracker due to lack of privilege!"));
+				}
+			}
+			else
+			{
+				FScopeLock GetCoordinateFramesLock(&Mutex);
+				uint32 NumPersistentFrames = 0;
+				MLResult Result = MLPersistentCoordinateFrameGetCount(FMagicLeapARPinInterface::Get().Pin()->GetHandle(), &NumPersistentFrames);
+				if (MLResult_Ok == Result && NumPersistentFrames > 0)
+				{
+					AllCoordinateFrames.Reset();
+					AllCoordinateFrames.AddZeroed(NumPersistentFrames);
+					MLCoordinateFrameUID* ArrayDataPointer = AllCoordinateFrames.GetData();
+					LastQueryResult = MLPersistentCoordinateFrameGetAll(FMagicLeapARPinInterface::Get().Pin()->GetHandle(), sizeof(MLCoordinateFrameUID) * NumPersistentFrames, &ArrayDataPointer);
+					if (MLResult_Ok != LastQueryResult)
+					{
+						UE_LOG(LogMagicLeap, Error, TEXT("MLPersistentCoordinateFrameGetAll failed with error %s"), UTF8_TO_TCHAR(MLPersistentCoordinateFrameGetResultString(LastQueryResult)));
+					}
+				}
+			}
+#endif //WITH_MLSDK
+
+			float SleepDuration;
+			{
+				FScopeLock ThreadSleepLock(&Mutex);
+				SleepDuration = ThreadSleepTime;
+			}
+			FPlatformProcess::Sleep(SleepDuration);
+		}
+
+		DestroyTracker();
+
+		return 0;
+	}
+
+	virtual void OnAppPause() override
+	{
+		// need to destroy the tracker here in case privileges are removed while the app is dormant.
+		DestroyTracker();
+	}
+
+	virtual void OnAppShutDown() override
+	{
+		DestroyTracker();
+	}
+
+private:
+	FMagicLeapARPinInterface()
+#if WITH_MLSDK
+	: MagicLeap::IAppEventHandler({ MLPrivilegeID_PwFoundObjRead })
+	, ThreadSleepTime(5.0f)
+	, Thread(nullptr)
+	, StopTaskCounter(0)
+	, Tracker(ML_INVALID_HANDLE)
+#endif //WITH_MLSDK
+	{
+		// TODO: Default affinity for this runnable thread on Lumin should be FLuminAffinity::GetPoolThreadMask(). Confirm and remove #if block.
+#if PLATFORM_LUMIN
+		Thread = FRunnableThread::Create(this, TEXT("PersistentCFWorker"), 0, TPri_BelowNormal, FLuminAffinity::GetPoolThreadMask());
+#else
+		Thread = FRunnableThread::Create(this, TEXT("PersistentCFWorker"), 0, TPri_BelowNormal);
+#endif // PLATFORM_LUMIN
+	}
+
+	void DestroyTracker()
+	{
+#if WITH_MLSDK
+		if (MLHandleIsValid(Tracker))
+		{
+			MLResult Result = MLPersistentCoordinateFrameTrackerDestroy(Tracker);
+			UE_CLOG(Result != MLResult_Ok, LogMagicLeap, Error, TEXT("Failed to destroy persistent coordinate frame tracker with error %s."), UTF8_TO_TCHAR(MLPersistentCoordinateFrameGetResultString(Result)));
+			Tracker = ML_INVALID_HANDLE;
+		}
+#endif //WITH_MLSDK
+	}
+
+	float ThreadSleepTime;
+	FRunnableThread* Thread;
+	FThreadSafeCounter StopTaskCounter;
+	mutable FCriticalSection Mutex;
+#if WITH_MLSDK
+	MLHandle Tracker;
+	TArray<MLCoordinateFrameUID> AllCoordinateFrames;
+	MLResult LastQueryResult;
+#endif //WITH_MLSDK
+
+	static TSharedPtr<FMagicLeapARPinInterface, ESPMode::ThreadSafe> Instance;
+};
+
+TSharedPtr<FMagicLeapARPinInterface, ESPMode::ThreadSafe> FMagicLeapARPinInterface::Instance;
+
+
+EPassableWorldError UMagicLeapARPinFunctionLibrary::GetNumAvailableARPins(int32& Count)
+{
+#if WITH_MLSDK
+	return FMagicLeapARPinInterface::Get().Pin()->GetNumAvailableARPins(Count);
+#else
+	return EPassableWorldError::Unavailable;
+#endif //WITH_MLSDK
+}
+
+EPassableWorldError UMagicLeapARPinFunctionLibrary::GetAvailableARPins(int32 NumRequested, TArray<FGuid>& Pins)
+{
+#if WITH_MLSDK
+	if (NumRequested <= 0)
+	{
+		GetNumAvailableARPins(NumRequested);
+	}
+
+	if (NumRequested == 0)
+	{
+		// There are no coordinate frames to return, so this call did succeed without any errors, it just returned an array of size 0.
+		Pins.Reset();
+		return EPassableWorldError::None;
+	}
+
+	MLResult Result = FMagicLeapARPinInterface::Get().Pin()->GetAvailableARPins(NumRequested, Pins);
+
+	return MLToUnrealPassableWorldError(Result);
+#else
+	return EPassableWorldError::Unavailable;
+#endif //WITH_MLSDK
+}
+
+EPassableWorldError UMagicLeapARPinFunctionLibrary::GetClosestARPin(const FVector& SearchPoint, FGuid& PinID)
+{
+#if WITH_MLSDK
+	return FMagicLeapARPinInterface::Get().Pin()->GetClosestARPin(SearchPoint, PinID);
+#else
+	return EPassableWorldError::Unavailable;
+#endif //WITH_MLSDK
+}
+
+bool UMagicLeapARPinFunctionLibrary::GetARPinPositionAndOrientation(const FGuid& PinID, FVector& Position, FRotator& Orientation)
+{
+	if (!IMagicLeapPlugin::Get().IsMagicLeapHMDValid())
+	{
+		return false;
+	}
+
+#if WITH_MLSDK
+	const MLCoordinateFrameUID* MLCFUID = reinterpret_cast<const MLCoordinateFrameUID*>(&PinID);
+	const FAppFramework& AppFramework = static_cast<FMagicLeapHMD*>(GEngine->XRSystem->GetHMDDevice())->GetAppFrameworkConst();
+	EFailReason FailReason = EFailReason::None;
+	FTransform Pose = FTransform::Identity;
+	if (AppFramework.GetTransform(*MLCFUID, Pose, FailReason))
+	{
+		FTransform TrackingToWorld = UHeadMountedDisplayFunctionLibrary::GetTrackingToWorldTransform(GWorld);
+		Pose.AddToTranslation(TrackingToWorld.GetLocation());
+		Pose.ConcatenateRotation(TrackingToWorld.Rotator().Quaternion());
+		Position = Pose.GetLocation();
+		Orientation = Pose.Rotator();
+		return true;
+	}
+#endif //WITH_MLSDK
+
+	return false;
+}
+
+float UMagicLeapARPinFunctionLibrary::GetARPinQueryInterval()
+{
+	return FMagicLeapARPinInterface::Get().Pin()->GetThreadSleepTime();
+}
+
+void UMagicLeapARPinFunctionLibrary::SetARPinQueryInterval(float QueryInterval)
+{
+	FMagicLeapARPinInterface::Get().Pin()->SetThreadSleepTime(QueryInterval);
+}
+
+UMagicLeapARPinComponent::UMagicLeapARPinComponent()
+: AutoPinType(EAutoPinType::OnlyOnDataRestoration)
+, bPinActor(false)
+, PinnedSceneComponent(nullptr)
+, ComponentRelativeToCFUID(FTransform::Identity)
+, bPinned(false)
+, bDataRestored(false)
+{
+	// Make sure this component ticks
+	PrimaryComponentTick.bCanEverTick = true;
+	PrimaryComponentTick.bStartWithTickEnabled = true;
+	PrimaryComponentTick.TickGroup = TG_PrePhysics;
+	bAutoActivate = true;
+}
+
+void UMagicLeapARPinComponent::BeginPlay()
+{
+	Super::BeginPlay();
+
+	if (ObjectUID.Len() == 0)
+	{
+		if (GetOwner() != nullptr)
+		{
+			ObjectUID = GetOwner()->GetName();
+			UE_LOG(LogMagicLeap, Warning, TEXT("ObjectUID is empty. Using Owner actor's name instead. A non-empty unique ID is required to make the object persistent."));
+		}
+	}
+
+	if (ObjectUID.Len() != 0)
+	{
+		FPersistentData Data;
+		bool bSuccess = UMagicLeapSecureStorage::GetSecureBlob<FPersistentData>(ObjectUID, Data);
+		if (bSuccess)
+		{
+#if WITH_MLSDK
+			FMemory::Memcpy(&PinnedCFUID, &Data.PinnedCFUID, sizeof(MLCoordinateFrameUID));
+#endif //WITH_MLSDK
+			ComponentRelativeToCFUID = Data.ComponentRelativeToCFUID;
+			bDataRestored = true;
+		}
+	}
+	else
+	{
+		UE_LOG(LogMagicLeap, Warning, TEXT("ObjectUID is empty. A non-empty unique ID is required to make the object persistent."));
+	}
+
+	if ((AutoPinType == EAutoPinType::Always) || (bDataRestored && AutoPinType == EAutoPinType::OnlyOnDataRestoration))
+	{
+		if (bPinActor)
+		{
+			PinActor(GetOwner());
+		}
+		else
+		{
+			PinSceneComponent(this);
+		}
+	}
+}
+
+void UMagicLeapARPinComponent::TickComponent(float DeltaTime, enum ELevelTick TickType, FActorComponentTickFunction *ThisTickFunction)
+{
+	Super::TickComponent(DeltaTime, TickType, ThisTickFunction);
+
+	if (!(IMagicLeapPlugin::Get().IsMagicLeapHMDValid()))
+	{
+		return;
+	}
+
+	if (PinnedSceneComponent != nullptr)
+	{
+		if (!bPinned)
+		{
+			if (bDataRestored)
+			{
+#if WITH_MLSDK
+				bPinned = FMagicLeapARPinInterface::Get().Pin()->IsCoordinateFrameValid(*reinterpret_cast<MLCoordinateFrameUID*>(&PinnedCFUID));
+#endif //WITH_MLSDK
+			}
+			else
+			{
+				EPassableWorldError Error = UMagicLeapARPinFunctionLibrary::GetClosestARPin(PinnedSceneComponent->GetComponentLocation(), PinnedCFUID);
+				if (Error == EPassableWorldError::None)
+				{
+					FVector PinWorldPosition = FVector::ZeroVector;
+					FRotator PinWorldOrientation = FRotator::ZeroRotator;
+					bool bCFUIDTransform = UMagicLeapARPinFunctionLibrary::GetARPinPositionAndOrientation(PinnedCFUID, PinWorldPosition, PinWorldOrientation);
+					if (bCFUIDTransform)
+					{
+						bPinned = true;
+						const FTransform& ComponentWorldTransform = PinnedSceneComponent->GetComponentToWorld();
+						const FVector ComponentRelativePosition = ComponentWorldTransform.GetLocation() - PinWorldPosition;
+						const FQuat ComponentRelativeOrientation = PinWorldOrientation.Quaternion().Inverse() * ComponentWorldTransform.GetRotation();
+						ComponentRelativeToCFUID = FTransform(ComponentRelativeOrientation, ComponentRelativePosition);
+
+						if (ObjectUID.Len() != 0)
+						{
+							FPersistentData PinData(PinnedCFUID, ComponentRelativeToCFUID);
+							UMagicLeapSecureStorage::PutSecureBlob<FPersistentData>(ObjectUID, &PinData);
+						}
+					}
+				}
+			}
+
+			if (bPinned)
+			{
+				OnPersistentEntityPinned.Broadcast(bDataRestored);
+			}
+		}
+
+		if (bPinned)
+		{
+			FVector PinWorldPosition = FVector::ZeroVector;
+			FRotator PinWorldOrientation = FRotator::ZeroRotator;
+			bool bCFUIDTransform = UMagicLeapARPinFunctionLibrary::GetARPinPositionAndOrientation(PinnedCFUID, PinWorldPosition, PinWorldOrientation);
+			if (bCFUIDTransform)
+			{
+				FTransform ComponentPinnedTransform = FTransform(PinWorldOrientation.Quaternion(), PinWorldPosition);
+				ComponentPinnedTransform = ComponentRelativeToCFUID * ComponentPinnedTransform;
+				PinnedSceneComponent->SetWorldLocationAndRotation(ComponentPinnedTransform.GetTranslation(), ComponentPinnedTransform.Rotator());
+			}
+		}
+	}
+}
+
+bool UMagicLeapARPinComponent::PinSceneComponent(USceneComponent* ComponentToPin)
+{
+	if (ComponentToPin != nullptr)
+	{
+		bPinned = (PinnedSceneComponent == ComponentToPin);
+		PinnedSceneComponent = ComponentToPin;
+	}
+	else
+	{
+		UE_LOG(LogMagicLeap, Warning, TEXT("nullptr passed to UMagicLeapARPinComponent::PinSceneComponent(). Use UMagicLeapARPinComponent::UnPin() if you no longer wish for this component to be persistent or want to move the component around."))
+	}
+	return (ComponentToPin != nullptr);
+}
+
+bool UMagicLeapARPinComponent::PinActor(AActor* ActorToPin)
+{
+	if (ActorToPin != nullptr)
+	{
+		return PinSceneComponent(ActorToPin->GetRootComponent());
+	}
+	return false;
+}
+
+void UMagicLeapARPinComponent::UnPin()
+{
+	PinnedSceneComponent = nullptr;
+	bPinned = false;
+	if (ObjectUID.Len() != 0)
+	{
+		UMagicLeapSecureStorage::DeleteSecureData(ObjectUID);
+		bDataRestored = false;
+	}
+}
+
+bool UMagicLeapARPinComponent::IsPinned() const
+{
+	return bPinned;
+}
+
+bool UMagicLeapARPinComponent::PinRestoredOrSynced() const
+{
+	return bDataRestored;
+}
+
+bool UMagicLeapARPinComponent::GetPinnedPinID(FGuid& PinID)
+{
+	if (bPinned)
+	{
+		PinID = PinnedCFUID;
+	}
+	return bPinned;
+}
