#DEV-VR - Merge from //UE4/Dev-MagicLeap/... @ CL 4136411

diff --git a/Engine/Plugins/Lumin/MagicLeap/Source/MagicLeap/Private/MagicLeapCustomPresent.cpp b/Engine/Plugins/Lumin/MagicLeap/Source/MagicLeap/Private/MagicLeapCustomPresent.cpp
--- a/Engine/Plugins/Lumin/MagicLeap/Source/MagicLeap/Private/MagicLeapCustomPresent.cpp
+++ b/Engine/Plugins/Lumin/MagicLeap/Source/MagicLeap/Private/MagicLeapCustomPresent.cpp
@@ -181,25 +181,28 @@ void FMagicLeapCustomPresentOpenGL::BeginRendering()
 #if WITH_MLSDK
 	check(IsInRenderingThread());
 
-	FTrackingFrame* frame = Plugin->GetCurrentFrame();
-	if (frame && bCustomPresentIsSet)
+	FTrackingFrame& frame = Plugin->GetCurrentFrameMutable();
+	if (bCustomPresentIsSet)
 	{
 		// TODO [Blake] : Need to see if we can use this newer matrix and override the view
 		// projection matrix (since they query GetStereoProjectionMatrix on the main thread)
-		MLStatus InitCameraStatus;
 		MLGraphicsFrameParams camera_params;
-		MLGraphicsInitFrameParams(&camera_params, &InitCameraStatus);
+		MLResult Result = MLGraphicsInitFrameParams(&camera_params);
+		if (Result != MLResult_Ok)
+		{
+			UE_LOG(LogMagicLeap, Error, TEXT("MLGraphicsInitFrameParams failed with status %d"), Result);
+		}
 		camera_params.projection_type = MLGraphicsProjectionType_ReversedInfiniteZ;
-		camera_params.surface_scale = frame->ScreenPercentage;
+		camera_params.surface_scale = frame.ScreenPercentage;
 		camera_params.protected_surface = false;
 		GConfig->GetBool(TEXT("/Script/LuminRuntimeSettings.LuminRuntimeSettings"), TEXT("bProtectedContent"), camera_params.protected_surface, GEngineIni);
 
 		// The near clipping plane is expected in meters despite what is documented in the header.
-		camera_params.near_clip = GNearClippingPlane / frame->WorldToMetersScale;
-		camera_params.far_clip = frame->FarClippingPlane / frame->WorldToMetersScale;
+		camera_params.near_clip = GNearClippingPlane / frame.WorldToMetersScale;
+		camera_params.far_clip = frame.FarClippingPlane / frame.WorldToMetersScale;
 		// Only focus distance equaling 1 engine unit seems to work on board without wearable and on desktop.
 #if PLATFORM_LUMIN
-		camera_params.focus_distance = frame->FocusDistance / frame->WorldToMetersScale;
+		camera_params.focus_distance = frame.FocusDistance / frame.WorldToMetersScale;
 #else
 		camera_params.focus_distance = 1.0f;
 #endif
@@ -208,42 +211,58 @@ void FMagicLeapCustomPresentOpenGL::BeginRendering()
 		if (BalanceCounter != 0)
 		{
 			UE_LOG(LogMagicLeap, Error, TEXT("Begin / End frame calls out of balance!"));
-			MLStatus OutStatus;
-			bool bResult = MLGraphicsSignalSyncObjectGL(Plugin->GraphicsClient, BalanceHandles[0], &OutStatus);
-			bResult = MLGraphicsSignalSyncObjectGL(Plugin->GraphicsClient, BalanceHandles[1], &OutStatus);
-			bResult = MLGraphicsEndFrame(Plugin->GraphicsClient, BalancePrevFrameHandle, &OutStatus);
+			MLResult Result = MLGraphicsSignalSyncObjectGL(Plugin->GraphicsClient, BalanceHandles[0]);
+			if (Result != MLResult_Ok)
+			{
+				UE_LOG(LogMagicLeap, Error, TEXT("MLGraphicsSignalSyncObjectGL for BalanceHandles[0] failed with status %d"), Result);
+			}
+
+			Result = MLGraphicsSignalSyncObjectGL(Plugin->GraphicsClient, BalanceHandles[1]);
+			if (Result != MLResult_Ok)
+			{
+				UE_LOG(LogMagicLeap, Error, TEXT("MLGraphicsSignalSyncObjectGL for BalanceHandles[1] failed with status %d"), Result);
+			}
+
+			Result = MLGraphicsEndFrame(Plugin->GraphicsClient, BalancePrevFrameHandle);
+			if (Result != MLResult_Ok)
+			{
+				UE_LOG(LogMagicLeap, Error, TEXT("MLGraphicsEndFrame failed with status %d"), Result);
+			}
 			--BalanceCounter;
 		}
 #endif
 
-		MLStatus OutStatus;
-		frame->bBeginFrameSucceeded = MLGraphicsBeginFrame(Plugin->GraphicsClient, &camera_params, &frame->Handle, &frame->RenderInfoArray, &OutStatus);
-		if (frame->bBeginFrameSucceeded)
+		Result = MLGraphicsBeginFrame(Plugin->GraphicsClient, &camera_params, &frame.Handle, &frame.RenderInfoArray);
+		frame.bBeginFrameSucceeded = (Result == MLResult_Ok);
+		if (frame.bBeginFrameSucceeded)
 		{
 #if BEGIN_END_FRAME_BALANCE_HACK
 			++BalanceCounter;
-			BalancePrevFrameHandle = frame->Handle;
-			BalanceHandles[0] = frame->RenderInfoArray.virtual_cameras[0].sync_object;
-			BalanceHandles[1] = frame->RenderInfoArray.virtual_cameras[1].sync_object;
+			BalancePrevFrameHandle = frame.Handle;
+			BalanceHandles[0] = frame.RenderInfoArray.virtual_cameras[0].sync_object;
+			BalanceHandles[1] = frame.RenderInfoArray.virtual_cameras[1].sync_object;
 #endif //BEGIN_END_FRAME_BALANCE_HACK
 
 			/* Convert eye extents from Graphics Projection Model to Unreal Projection Model */
 			// Unreal expects the projection matrix to be in centimeters and uses it for various purposes
 			// such as bounding volume calculations for lights in the shadow algorithm.
 			// We're overwriting the near value to match the units of unreal here instead of using the units of the SDK
-			for (uint32_t eye = 0; eye < frame->RenderInfoArray.num_virtual_cameras; ++eye)
+			for (uint32_t eye = 0; eye < frame.RenderInfoArray.num_virtual_cameras; ++eye)
 			{
-				frame->RenderInfoArray.virtual_cameras[eye].projection.matrix_colmajor[10] = 0.0f; // Model change hack
-				frame->RenderInfoArray.virtual_cameras[eye].projection.matrix_colmajor[11] = -1.0f; // Model change hack
-				frame->RenderInfoArray.virtual_cameras[eye].projection.matrix_colmajor[14] = GNearClippingPlane; // Model change hack
+				frame.RenderInfoArray.virtual_cameras[eye].projection.matrix_colmajor[10] = 0.0f; // Model change hack
+				frame.RenderInfoArray.virtual_cameras[eye].projection.matrix_colmajor[11] = -1.0f; // Model change hack
+				frame.RenderInfoArray.virtual_cameras[eye].projection.matrix_colmajor[14] = GNearClippingPlane; // Model change hack
 			}
 		}
 		else
 		{
-			UE_LOG(LogMagicLeap, Error, TEXT("MLGraphicsBeginFrame failed with status %d"), static_cast<int32>(OutStatus));
+			if (Result != MLResult_Timeout)
+			{
+				UE_LOG(LogMagicLeap, Error, TEXT("MLGraphicsBeginFrame failed with status %d"), Result);
+			}
 			// TODO: See if this is only needed for ZI.
-			frame->Handle = ML_INVALID_HANDLE;
-			MagicLeap::ResetVirtualCameraInfoArray(frame->RenderInfoArray);
+			frame.Handle = ML_INVALID_HANDLE;
+			MagicLeap::ResetVirtualCameraInfoArray(frame.RenderInfoArray);
 		}
 	}
 #endif //WITH_MLSDK
@@ -254,23 +273,27 @@ void FMagicLeapCustomPresentOpenGL::FinishRendering()
 #if WITH_MLSDK
 	check(IsInRenderingThread());
 
-	if (Plugin->IsDeviceInitialized() && Plugin->GetCurrentFrame()->bBeginFrameSucceeded)
+	if (Plugin->IsDeviceInitialized() && Plugin->GetCurrentFrame().bBeginFrameSucceeded)
 	{
 		if (bNotifyLifecycleOfFirstPresent)
 		{
-			MLLifecycleErrorCode result;
 			// Lifecycle tells the system's loading indicator to stop.
 			// App's rendering takes over.
-			bool bResult = MLLifecycleSetReadyIndication(&result);
-			bNotifyLifecycleOfFirstPresent = !bResult;
-			if (!bResult || result != MLLifecycleErrorCode_Success)
+			MLResult Result = MLLifecycleSetReadyIndication();
+			bNotifyLifecycleOfFirstPresent = (Result != MLResult_Ok);
+			if (Result != MLResult_Ok)
+			{
+				UE_LOG(LogMagicLeap, Error, TEXT("Error (%d) sending app ready indication to lifecycle."), Result);
+			}
+			else
 			{
-				UE_LOG(LogMagicLeap, Error, TEXT("Error sending app ready indication to lifecycle."));
+				// [temporary] used for KPI tracking.
+				UE_LOG(LogMagicLeap, Display, TEXT("Presenting first render from app."));
 			}
 		}
 
 		// TODO [Blake] : Hack since we cannot yet specify a handle per view in the view family
-		const MLGraphicsVirtualCameraInfoArray& vp_array = Plugin->GetCurrentFrame()->RenderInfoArray;
+		const MLGraphicsVirtualCameraInfoArray& vp_array = Plugin->GetCurrentFrame().RenderInfoArray;
 		const uint32 vp_width = static_cast<uint32>(vp_array.viewport.w);
 		const uint32 vp_height = static_cast<uint32>(vp_array.viewport.h);
 
@@ -290,7 +313,7 @@ void FMagicLeapCustomPresentOpenGL::FinishRendering()
 			glDisable(GL_FRAMEBUFFER_SRGB);
 		}
 
-		check(vp_array.num_virtual_cameras >= 2); // We assume at least one virtual camera per eye
+		//check(vp_array.num_virtual_cameras >= 2); // We assume at least one virtual camera per eye
 
 		FVector2D InternalTextureDims;
 		Plugin->GetAppFrameworkConst().GetDeviceResolution(InternalTextureDims);
@@ -306,19 +329,18 @@ void FMagicLeapCustomPresentOpenGL::FinishRendering()
 		glBindFramebuffer(GL_DRAW_FRAMEBUFFER, Framebuffers[1]);
 		FOpenGL::BlitFramebuffer(0, 0, InternalTextureDims.X / 2, InternalTextureDims.Y, 0, 0, vp_width, vp_height, GL_COLOR_BUFFER_BIT, GL_NEAREST);
 
-		MLStatus OutStatus;
-		bool bResult = MLGraphicsSignalSyncObjectGL(Plugin->GraphicsClient, vp_array.virtual_cameras[0].sync_object, &OutStatus);
-		if (!bResult)
+		MLResult Result = MLGraphicsSignalSyncObjectGL(Plugin->GraphicsClient, vp_array.virtual_cameras[0].sync_object);
+		if (Result != MLResult_Ok)
 		{
-			UE_LOG(LogMagicLeap, Error, TEXT("MLGraphicsSignalSyncObjectGL for eye 0 failed with status %d"), static_cast<int32>(OutStatus));
+			UE_LOG(LogMagicLeap, Error, TEXT("MLGraphicsSignalSyncObjectGL for eye 0 failed with status %d"), Result);
 		}
 
 		FOpenGL::FramebufferTextureLayer(GL_FRAMEBUFFER, GL_COLOR_ATTACHMENT0, vp_array.color_id, 0, 1);
 		FOpenGL::BlitFramebuffer(InternalTextureDims.X / 2, 0, InternalTextureDims.X, InternalTextureDims.Y, 0, 0, vp_width, vp_height, GL_COLOR_BUFFER_BIT, GL_NEAREST);
-		bResult = MLGraphicsSignalSyncObjectGL(Plugin->GraphicsClient, vp_array.virtual_cameras[1].sync_object, &OutStatus);
-		if (!bResult)
+		Result = MLGraphicsSignalSyncObjectGL(Plugin->GraphicsClient, vp_array.virtual_cameras[1].sync_object);
+		if (Result != MLResult_Ok)
 		{
-			UE_LOG(LogMagicLeap, Error, TEXT("MLGraphicsSignalSyncObjectGL for eye 1 failed with status %d"), static_cast<int32>(OutStatus));
+			UE_LOG(LogMagicLeap, Error, TEXT("MLGraphicsSignalSyncObjectGL for eye 1 failed with status %d"), Result);
 		}
 
 		glBindFramebuffer(GL_DRAW_FRAMEBUFFER, CurrentFB);
@@ -331,10 +353,10 @@ void FMagicLeapCustomPresentOpenGL::FinishRendering()
 #if 0 // Enable this in case the MLSDK increases the size of the virtual_cameras array past 2
 		for (uint32 i = 2; i < vp_array.num_virtual_cameras; ++i)
 		{
-			bResult = MLGraphicsSignalSyncObjectGL(Plugin->GraphicsClient, vp_array.virtual_cameras[i].sync_object, &OutStatus);
-			if (!bResult)
+			Result = MLGraphicsSignalSyncObjectGL(Plugin->GraphicsClient, vp_array.virtual_cameras[i].sync_object);
+			if (Result != MLResult_Ok)
 			{
-				UE_LOG(LogMagicLeap, Error, TEXT("MLGraphicsSignalSyncObjectGL for eye %d failed with status %d"), i, static_cast<int32>(OutStatus));
+				UE_LOG(LogMagicLeap, Error, TEXT("MLGraphicsSignalSyncObjectGL for eye %d failed with status %d"), i, Result);
 			}
 		}
 #endif
@@ -343,10 +365,10 @@ void FMagicLeapCustomPresentOpenGL::FinishRendering()
 		--BalanceCounter;
 #endif //BEGIN_END_FRAME_BALANCE_HACK
 
-		bResult = MLGraphicsEndFrame(Plugin->GraphicsClient, Plugin->GetCurrentFrame()->Handle, &OutStatus);
-		if (!bResult)
+		Result = MLGraphicsEndFrame(Plugin->GraphicsClient, Plugin->GetCurrentFrame().Handle);
+		if (Result != MLResult_Ok)
 		{
-			UE_LOG(LogMagicLeap, Error, TEXT("MLGraphicsEndFrame failed with status %d"), static_cast<int32>(OutStatus));
+			UE_LOG(LogMagicLeap, Error, TEXT("MLGraphicsEndFrame failed with status %d"), Result);
 		}
 	}
 	Plugin->InitializeOldFrameFromRenderFrame();
@@ -439,72 +461,82 @@ void FMagicLeapCustomPresentVulkan::BeginRendering()
 {
 	check(IsInRenderingThread());
 
-	FTrackingFrame* frame = Plugin->GetCurrentFrame();
-	if (frame && bCustomPresentIsSet)
+	FTrackingFrame& frame = Plugin->GetCurrentFrameMutable();
+	if (bCustomPresentIsSet)
 	{
 		// TODO [Blake] : Need to see if we can use this newer matrix and override the view
 		// projection matrix (since they query GetStereoProjectionMatrix on the main thread)
-		MLStatus InitCameraStatus;
 		MLGraphicsFrameParams camera_params;
-		MLGraphicsInitFrameParams(&camera_params, &InitCameraStatus);
+		MLResult InitResult = MLGraphicsInitFrameParams(&camera_params);
+		if (InitResult != MLResult_Ok)
+		{
+			UE_LOG(LogMagicLeap, Error, TEXT("MLGraphicsInitFrameParams failed with status %d"), InitResult);
+		}
 		camera_params.projection_type = MLGraphicsProjectionType_ReversedInfiniteZ;
 		camera_params.surface_scale = 1.0f;
 		camera_params.protected_surface = false;
 		GConfig->GetBool(TEXT("/Script/LuminRuntimeSettings.LuminRuntimeSettings"), TEXT("bProtectedContent"), camera_params.protected_surface, GEngineIni);
 
 		// The near clipping plane is expected in meters despite what is documented in the header.
-		camera_params.near_clip = GNearClippingPlane / frame->WorldToMetersScale;
-		camera_params.far_clip = frame->FarClippingPlane / frame->WorldToMetersScale;
+		camera_params.near_clip = GNearClippingPlane / frame.WorldToMetersScale;
+		camera_params.far_clip = frame.FarClippingPlane / frame.WorldToMetersScale;
 
 		// The focus distance is expected in meters despite what is documented in the header.
 		// Only focus distance equaling 1 engine unit seems to work on board without wearable and on desktop.
 #if PLATFORM_LUMIN
-		camera_params.focus_distance = frame->FocusDistance / frame->WorldToMetersScale;
+		camera_params.focus_distance = frame.FocusDistance / frame.WorldToMetersScale;
 #else
 		camera_params.focus_distance = 1.0f;
 #endif
 
 	  ExecuteOnRHIThread_DoNotWait([this, camera_params]()
 	  {
-		  FTrackingFrame* RHIframe = Plugin->GetCurrentFrame();
-		  MLStatus OutStatus;
+		  FTrackingFrame& RHIframe = Plugin->GetCurrentFrameMutable();
 
 #if BEGIN_END_FRAME_BALANCE_HACK
 			if (BalanceCounter != 0)
 			{
 				UE_LOG(LogMagicLeap, Error, TEXT("Begin / End frame calls out of balance!"));
 				FMagicLeapHelperVulkan::SignalObjects((uint64)BalanceHandles[0], (uint64)BalanceHandles[1]);
-				bool bResult = MLGraphicsEndFrame(Plugin->GraphicsClient, BalancePrevFrameHandle, &OutStatus);
+				MLResult Result = MLGraphicsEndFrame(Plugin->GraphicsClient, BalancePrevFrameHandle);
+				if (Result != MLResult_Ok)
+				{
+					UE_LOG(LogMagicLeap, Error, TEXT("MLGraphicsEndFrame failed with status %d"), Result);
+				}
 				--BalanceCounter;
 			}
 #endif
 
-		  RHIframe->bBeginFrameSucceeded = MLGraphicsBeginFrame(Plugin->GraphicsClient, &camera_params, &RHIframe->Handle, &RHIframe->RenderInfoArray, &OutStatus);
-		  if (RHIframe->bBeginFrameSucceeded)
+		  MLResult Result = MLGraphicsBeginFrame(Plugin->GraphicsClient, &camera_params, &RHIframe.Handle, &RHIframe.RenderInfoArray);
+		  RHIframe.bBeginFrameSucceeded = (Result == MLResult_Ok);
+		  if (RHIframe.bBeginFrameSucceeded)
 		  {
 #if BEGIN_END_FRAME_BALANCE_HACK
 				++BalanceCounter;
-				BalancePrevFrameHandle = RHIframe->Handle;
-				BalanceHandles[0] = RHIframe->RenderInfoArray.virtual_cameras[0].sync_object;
-				BalanceHandles[1] = RHIframe->RenderInfoArray.virtual_cameras[1].sync_object;
+				BalancePrevFrameHandle = RHIframe.Handle;
+				BalanceHandles[0] = RHIframe.RenderInfoArray.virtual_cameras[0].sync_object;
+				BalanceHandles[1] = RHIframe.RenderInfoArray.virtual_cameras[1].sync_object;
 #endif //BEGIN_END_FRAME_BALANCE_HACK
 			  /* Convert eye extents from Graphics Projection Model to Unreal Projection Model */
 			  // Unreal expects the projection matrix to be in centimeters and uses it for various purposes
 			  // such as bounding volume calculations for lights in the shadow algorithm.
 			  // We're overwriting the near value to match the units of unreal here instead of using the units of the SDK
-			  for (uint32_t eye = 0; eye < RHIframe->RenderInfoArray.num_virtual_cameras; ++eye)
+			  for (uint32_t eye = 0; eye < RHIframe.RenderInfoArray.num_virtual_cameras; ++eye)
 			  {
-				  RHIframe->RenderInfoArray.virtual_cameras[eye].projection.matrix_colmajor[10] = 0.0f; // Model change hack
-				  RHIframe->RenderInfoArray.virtual_cameras[eye].projection.matrix_colmajor[11] = -1.0f; // Model change hack
-				  RHIframe->RenderInfoArray.virtual_cameras[eye].projection.matrix_colmajor[14] = GNearClippingPlane; // Model change hack
+				  RHIframe.RenderInfoArray.virtual_cameras[eye].projection.matrix_colmajor[10] = 0.0f; // Model change hack
+				  RHIframe.RenderInfoArray.virtual_cameras[eye].projection.matrix_colmajor[11] = -1.0f; // Model change hack
+				  RHIframe.RenderInfoArray.virtual_cameras[eye].projection.matrix_colmajor[14] = GNearClippingPlane; // Model change hack
 			  }
 		  }
 		  else
 		  {
-			  UE_LOG(LogMagicLeap, Error, TEXT("MLGraphicsBeginFrame failed with status %d"), static_cast<int32>(OutStatus));
+			  if (Result != MLResult_Timeout)
+			  {
+				UE_LOG(LogMagicLeap, Error, TEXT("MLGraphicsBeginFrame failed with status %d"), Result);
+			  }
 			  // TODO: See if this is only needed for ZI.
-			  RHIframe->Handle = ML_INVALID_HANDLE;
-			  MagicLeap::ResetVirtualCameraInfoArray(RHIframe->RenderInfoArray);
+			  RHIframe.Handle = ML_INVALID_HANDLE;
+			  MagicLeap::ResetVirtualCameraInfoArray(RHIframe.RenderInfoArray);
 		  }
 	  });
   }
@@ -514,7 +546,7 @@ void FMagicLeapCustomPresentVulkan::FinishRendering()
 {
 	check(IsInRenderingThread() || IsInRHIThread());
 
-	if (Plugin->IsDeviceInitialized() && Plugin->GetCurrentFrame()->bBeginFrameSucceeded)
+	if (Plugin->IsDeviceInitialized() && Plugin->GetCurrentFrame().bBeginFrameSucceeded)
 	{
 #if !PLATFORM_LUMIN
 		// Set up the viewport
@@ -533,11 +565,10 @@ void FMagicLeapCustomPresentVulkan::FinishRendering()
 #endif
 
 		// TODO [Blake] : Hack since we cannot yet specify a handle per view in the view family
-		const MLGraphicsVirtualCameraInfoArray& vp_array = Plugin->GetCurrentFrame()->RenderInfoArray;
+		const MLGraphicsVirtualCameraInfoArray& vp_array = Plugin->GetCurrentFrame().RenderInfoArray;
 		const uint32 vp_width = static_cast<uint32>(vp_array.viewport.w);
 		const uint32 vp_height = static_cast<uint32>(vp_array.viewport.h);
 
-		MLStatus OutStatus;
 		bool bTestClear = false;
 		if (bTestClear)
 		{
@@ -566,10 +597,10 @@ void FMagicLeapCustomPresentVulkan::FinishRendering()
 		--BalanceCounter;
 #endif //BEGIN_END_FRAME_BALANCE_HACK
 
-		bool bResult = MLGraphicsEndFrame(Plugin->GraphicsClient, Plugin->GetCurrentFrame()->Handle, &OutStatus);
-		if (!bResult)
+		MLResult Result = MLGraphicsEndFrame(Plugin->GraphicsClient, Plugin->GetCurrentFrame().Handle);
+		if (Result != MLResult_Ok)
 		{
-			UE_LOG(LogMagicLeap, Error, TEXT("MLGraphicsEndFrame failed with status %d"), static_cast<int32>(OutStatus));
+			UE_LOG(LogMagicLeap, Error, TEXT("MLGraphicsEndFrame failed with status %d"), Result);
 		}
 	}
   
@@ -632,14 +663,16 @@ bool FMagicLeapCustomPresentVulkan::Present(int& SyncInterval)
 
 	if (bNotifyLifecycleOfFirstPresent)
 	{
-		MLLifecycleErrorCode result;
-		// Lifecycle tells the system's loading indicator to stop.
-		// App's rendering takes over.
-		bool bResult = MLLifecycleSetReadyIndication(&result);
-		bNotifyLifecycleOfFirstPresent = !bResult;
-		if (!bResult || result != MLLifecycleErrorCode_Success)
+		MLResult Result = MLLifecycleSetReadyIndication();
+		bNotifyLifecycleOfFirstPresent = (Result != MLResult_Ok);
+		if (Result != MLResult_Ok)
+		{
+			UE_LOG(LogMagicLeap, Error, TEXT("MLLifecycleSetReadyIndication failed with error %d."), Result);
+		}
+		else
 		{
-			UE_LOG(LogMagicLeap, Error, TEXT("Error sending app ready indication to lifecycle."));
+			// [temporary] used for KPI tracking.
+			UE_LOG(LogMagicLeap, Display, TEXT("Presenting first render from app."));
 		}
 	}
 
diff --git a/Engine/Plugins/Lumin/MagicLeap/Source/MagicLeap/Private/MagicLeapCustomPresent.h b/Engine/Plugins/Lumin/MagicLeap/Source/MagicLeap/Private/MagicLeapCustomPresent.h
--- a/Engine/Plugins/Lumin/MagicLeap/Source/MagicLeap/Private/MagicLeapCustomPresent.h
+++ b/Engine/Plugins/Lumin/MagicLeap/Source/MagicLeap/Private/MagicLeapCustomPresent.h
@@ -26,6 +26,7 @@ ML_INCLUDES_END
 
 class FMagicLeapHMD;
 class FViewport;
+struct FWorldContext;
 
 static const uint8 kNumEyes = 2;
 
@@ -56,6 +57,7 @@ public:
 
 	// [0.0f, 1.0f]
 	float ScreenPercentage;
+	FWorldContext* WorldContext;
 
 	FTrackingFrame()
 		: FrameNumber(0)
@@ -76,6 +78,7 @@ public:
 		, Handle(ML_INVALID_HANDLE)
 #endif //WITH_MLSDK
 		, ScreenPercentage(1.0f)
+		, WorldContext(nullptr)
 	{
 #if WITH_MLSDK
 		FrameId.data[0] = 0;
diff --git a/Engine/Plugins/Lumin/MagicLeap/Source/MagicLeap/Private/MagicLeapHMD.cpp b/Engine/Plugins/Lumin/MagicLeap/Source/MagicLeap/Private/MagicLeapHMD.cpp
--- a/Engine/Plugins/Lumin/MagicLeap/Source/MagicLeap/Private/MagicLeapHMD.cpp
+++ b/Engine/Plugins/Lumin/MagicLeap/Source/MagicLeap/Private/MagicLeapHMD.cpp
@@ -2,6 +2,8 @@
 
 #include "MagicLeapHMD.h"
 #include "MagicLeapHMDFunctionLibrary.h"
+#include "DefaultStereoLayers.h"
+#include "MagicLeapStereoLayers.h"
 #include "RendererPrivate.h"
 #include "PostProcess/PostProcessHMD.h"
 #include "ScreenRendering.h"
@@ -43,7 +45,7 @@ ML_INCLUDES_START
 #include <ml_perception.h>
 
 #if !PLATFORM_LUMIN
-#include <ml_virtual_device.h>
+#include <ml_remote.h>
 #include "Misc/MessageDialog.h"
 #endif
 
@@ -58,6 +60,8 @@ ML_INCLUDES_START
 #include <ml_graphics.h>
 #endif // PLATFORM_LUMIN
 
+#include "ml_privileges.h"
+
 ML_INCLUDES_END
 #endif //WITH_MLSDK
 
@@ -98,6 +102,7 @@ public:
 		APISetup.LoadDLL(TEXT("ml_perception_client"));
 		APISetup.LoadDLL(TEXT("ml_graphics"));
 		APISetup.LoadDLL(TEXT("ml_lifecycle"));
+		APISetup.LoadDLL(TEXT("ml_privileges"));
 #endif
 
 		if (bIsVDZIEnabled)
@@ -105,7 +110,7 @@ public:
 #if PLATFORM_WINDOWS
 			UE_LOG(LogMagicLeap, Log, TEXT("ML VDZI mode enabled. Using OpenGL renderer."));
 
-			APISetup.LoadDLL(TEXT("ml_virtual_device"));
+			APISetup.LoadDLL(TEXT("ml_remote"));
 
 			// DirectX, which is currently not supported by MagicLeap, is default API on Windows.
 			// OpenGL is forced by loading module in PostConfigInit phase and passing in command line.
@@ -143,26 +148,24 @@ public:
 		return LocalHMD;
 	}
 
-#if ENGINE_MAJOR_VERSION >= 4 && ENGINE_MINOR_VERSION <= 13
-	FString GetModulePriorityKeyName() const
-#else
 	FString GetModuleKeyName() const
-#endif // ENGINE_MAJOR_VERSION >= 4 && ENGINE_MINOR_VERSION <= 13
 	{
 		return FString(TEXT("MagicLeap"));
 	}
 
-	virtual bool IsMagicLeapHMDValidAndConnected() const override
+	virtual bool IsMagicLeapHMDValid() const override
 	{
+#if PLATFORM_LUMIN
+		return true;
+#else
 		if (IsValid(GEngine) && GEngine->XRSystem.IsValid() && (GEngine->XRSystem->GetSystemName().Compare(FMagicLeapHMD::SystemName) == 0))
 		{
 			TSharedPtr<FMagicLeapHMD, ESPMode::ThreadSafe> hmd = StaticCastSharedPtr<FMagicLeapHMD>(GEngine->XRSystem);
-			if (hmd.IsValid())
-			{
-				return hmd->IsHMDConnected();
-			}
+			// IsHMDConnected() is an expensive call when MLRemote is enabled, so we'll keep the onus of that check on the caller.
+			return hmd.IsValid();
 		}
 		return false;
+#endif // PLATFORM_LUMIN
 	}
 
 	virtual TWeakPtr<IMagicLeapHMD, ESPMode::ThreadSafe> GetHMD() override
@@ -243,9 +246,12 @@ bool FMagicLeapHMD::OnStartGameFrame(FWorldContext& WorldContext)
 #if WITH_MLSDK
 	check(IsInGameThread());
 
+	FTrackingFrame& TrackingFrame = GetCurrentFrameMutable();
+
 	if (!WorldContext.World() || !WorldContext.World()->IsGameWorld())
 	{
 		// ignore all non-game worlds
+		TrackingFrame.WorldContext = nullptr;
 		return false;
 	}
 
@@ -264,17 +270,19 @@ bool FMagicLeapHMD::OnStartGameFrame(FWorldContext& WorldContext)
 
 	AppFramework.BeginUpdate();
 
-	FTrackingFrame* TrackingFrame = GetCurrentFrame();
-	check(TrackingFrame);
 	// init tracking frame if first frame, otherwise we keep using last frame's data until it is refreshed in BeginRendering_GameThread
-	if (TrackingFrame->Snapshot == nullptr)
+	if (TrackingFrame.Snapshot == nullptr)
 	{
 		RefreshTrackingFrame();
 	}
 
 	// override the default value that the frame constructor initialized and make sure it is non zero
-	TrackingFrame->WorldToMetersScale = WorldContext.World()->GetWorldSettings()->WorldToMeters;
-	TrackingFrame->WorldToMetersScale = TrackingFrame->WorldToMetersScale == 0.0f ? 100.0f : TrackingFrame->WorldToMetersScale;
+	TrackingFrame.WorldToMetersScale = WorldContext.World()->GetWorldSettings()->WorldToMeters;
+	TrackingFrame.WorldToMetersScale = TrackingFrame.WorldToMetersScale == 0.0f ? 100.0f : TrackingFrame.WorldToMetersScale;
+	TrackingFrame.WorldContext = &WorldContext;
+
+	RefreshTrackingToWorldTransform(WorldContext);
+
 #endif //WITH_MLSDK
 
 	return true;
@@ -302,7 +310,17 @@ bool FMagicLeapHMD::IsHMDConnected()
 	bool bZIServerRunning = false;
 	if (bIsVDZIEnabled)
 	{
-		ML_FUNCTION_WRAPPER(bZIServerRunning = MLVirtualDeviceIsServerRunning())
+		if (FMagicLeapSDKDetection::IsSDKDetected())
+		{
+			MLResult Result = MLRemoteIsServerConfigured(&bZIServerRunning);
+			if (Result != MLResult_Ok)
+			{
+				UE_LOG(LogMagicLeap, Error, TEXT("MLRemoteIsServerConfigured failed with error %d!"), Result);
+
+				// Ensure we don't falsely mark it as running if there was an error
+				bZIServerRunning = false;
+			}
+		}
 	}
 	// TODO: MLVirtualDeviceIsServerRunning() crashes when called on render thread.
 	return AppFramework.IsInitialized() && bIsVDZIEnabled && bZIServerRunning;
@@ -317,16 +335,7 @@ bool FMagicLeapHMD::IsHMDConnected()
 
 bool FMagicLeapHMD::IsHMDEnabled() const
 {
-#if WITH_EDITOR
-	// We indicate the HMD is disable in editor so that the vr-editor mode is disabled.
-	// We do this by indicating it's enabled only when playing (i.e. when OnBeginPlay
-	// was called) and only when executing a console command (which is only indicative
-	// when the GStatProcessingViewportClient is valid -- See
-	// UUnrealEdEngine::Exec() for why this is the case).
-	return bHmdEnabled && bIsPlaying && (nullptr != GStatProcessingViewportClient);
-#else
 	return bHmdEnabled;
-#endif
 }
 
 void FMagicLeapHMD::EnableHMD(bool Enable)
@@ -369,12 +378,9 @@ bool FMagicLeapHMD::GetHMDMonitorInfo(MonitorInfo& MonitorDesc)
 
 void FMagicLeapHMD::GetFieldOfView(float& OutHFOVInDegrees, float& OutVFOVInDegrees) const
 {
-	FTrackingFrame* frame = GetCurrentFrame();
-	if (frame)
-	{
-		OutHFOVInDegrees = frame->HFov;
-		OutVFOVInDegrees = frame->VFov;
-	}
+	const FTrackingFrame& frame = GetCurrentFrame();
+	OutHFOVInDegrees = frame.HFov;
+	OutVFOVInDegrees = frame.VFov;
 }
 
 bool FMagicLeapHMD::DoesSupportPositionalTracking() const
@@ -384,33 +390,29 @@ bool FMagicLeapHMD::DoesSupportPositionalTracking() const
 
 bool FMagicLeapHMD::HasValidTrackingPosition()
 {
-	FTrackingFrame* frame = GetCurrentFrame();
-	return bHmdPosTracking && frame ? frame->HasHeadTrackingPosition : false;
+	const FTrackingFrame& frame = GetCurrentFrame();
+	return bHmdPosTracking ? frame.HasHeadTrackingPosition : false;
 }
 
 bool FMagicLeapHMD::GetTrackingSensorProperties(int32 DeviceId, FQuat& OutOrientation, FVector& OutPosition, FXRSensorProperties& OutSensorProperties)
 {
 	// Assuming there is only one tracker (the device itself) on the system being optically tracked
-	FTrackingFrame* frame = GetCurrentFrame();
-	if (frame)
-	{
-		const float HalfHFOV = frame->HFov / 2.f;
-		const float HalfVFOV = frame->VFov / 2.f;
+	const FTrackingFrame& frame = GetCurrentFrame();
+	const float HalfHFOV = frame.HFov / 2.f;
+	const float HalfVFOV = frame.VFov / 2.f;
 
-		OutSensorProperties.TopFOV = HalfVFOV;
-		OutSensorProperties.BottomFOV = HalfVFOV;
+	OutSensorProperties.TopFOV = HalfVFOV;
+	OutSensorProperties.BottomFOV = HalfVFOV;
 
-		OutSensorProperties.LeftFOV = HalfHFOV;
-		OutSensorProperties.RightFOV = HalfHFOV;
+	OutSensorProperties.LeftFOV = HalfHFOV;
+	OutSensorProperties.RightFOV = HalfHFOV;
 
-		OutSensorProperties.CameraDistance = 0.f;
-		OutSensorProperties.NearPlane = 8.f;
-		OutSensorProperties.FarPlane = 400.f; // Assumption, should get real numbers on this!
+	// TODO: set correct values here.
+	OutSensorProperties.CameraDistance = 0.f;
+	OutSensorProperties.NearPlane = 8.f;
+	OutSensorProperties.FarPlane = 400.f; // Assumption, should get real numbers on this!
 
-		return true;
-	}
-
-	return false;
+	return true;
 }
 
 void FMagicLeapHMD::SetInterpupillaryDistance(float NewInterpupillaryDistance)
@@ -430,11 +432,7 @@ bool FMagicLeapHMD::IsChromaAbCorrectionEnabled() const
 
 bool FMagicLeapHMD::IsHeadTrackingAllowed() const
 {
-	const FTrackingFrame* frame = GetCurrentFrame();
-	if (!frame)
-	{
-		return false;
-	}
+	const FTrackingFrame& frame = GetCurrentFrame();
 #if WITH_EDITOR
 	if (GIsEditor)
 	{
@@ -453,15 +451,9 @@ void FMagicLeapHMD::ResetOrientationAndPosition(float yaw)
 
 void FMagicLeapHMD::ResetOrientation(float Yaw)
 {
-	const FTrackingFrame* frame = GetCurrentFrame();
+	const FTrackingFrame& frame = GetCurrentFrame();
 
-	FTransform RawTransform;
-	if (frame)
-	{
-		RawTransform = frame->RawPose;
-	}
-
-	FRotator ViewRotation = RawTransform.Rotator();
+	FRotator ViewRotation = frame.RawPose.Rotator();
 	ViewRotation.Pitch = 0;
 	ViewRotation.Roll = 0;
 
@@ -477,16 +469,8 @@ void FMagicLeapHMD::ResetOrientation(float Yaw)
 
 void FMagicLeapHMD::ResetPosition()
 {
-	const FTrackingFrame* frame = GetCurrentFrame();
-
-	FTransform RawTransform;
-	if (frame)
-	{
-		RawTransform = frame->RawPose;
-	}
-
-	FVector ViewPositon = RawTransform.GetTranslation();
-	AppFramework.SetBasePosition(ViewPositon);
+	const FTrackingFrame& frame = GetCurrentFrame();
+	AppFramework.SetBasePosition(frame.RawPose.GetTranslation());
 }
 
 void FMagicLeapHMD::SetBasePosition(const FVector& InBasePosition)
@@ -551,12 +535,19 @@ void FMagicLeapHMD::RefreshTrackingFrame()
 	// Release the snapshot of the previous frame.
 	// This is done here instead of on end frame because modules implemented as input devices (Gestures, controller)
 	// are ticked and fire their events before the OnStartGameFrame().
-	MLPerceptionReleaseSnapshot(GameTrackingFrame.Snapshot);
+	MLResult Result = MLPerceptionReleaseSnapshot(GameTrackingFrame.Snapshot);
+#if PLATFORM_LUMIN
+	UE_CLOG(Result != MLResult_Ok, LogMagicLeap, Error, TEXT("MLImageTrackerInitSettings failed with error %d."), Result);
+#endif //PLATFORM_LUMIN
+
 	// get the snapshot for the frame
-	GameTrackingFrame.Snapshot = MLPerceptionGetSnapshot();
+	Result = MLPerceptionGetSnapshot(&GameTrackingFrame.Snapshot);
+#if PLATFORM_LUMIN
+	UE_CLOG(Result != MLResult_Ok, LogMagicLeap, Error, TEXT("MLPerceptionGetSnapshot failed with error %d."), Result);
+#endif //PLATFORM_LUMIN
 
 	MLHeadTrackingState state;
-	bHeadTrackingStateAvailable = MLHeadTrackingGetState(HeadTracker, &state);
+	bHeadTrackingStateAvailable = MLHeadTrackingGetState(HeadTracker, &state) == MLResult_Ok;
 	if (bHeadTrackingStateAvailable)
 	{
 		HeadTrackingState.Mode = MLToUnrealHeadTrackingMode(state.mode);
@@ -698,10 +689,8 @@ void FMagicLeapHMD::SetFrameTimingHint(ELuminFrameTimingHint InFrameTimingHint)
 				UE_LOG(LogMagicLeap, Warning, TEXT("Tried to set invalid Frame Timing Hint!  Defaulting to unspecified."));
 			}
 
-			MLStatus RetStatus;
-			MLGraphicsSetFrameTimingHint(GraphicsClient, FTH, &RetStatus);
-
-			if (RetStatus == MLStatus_OK)
+			MLResult Result = MLGraphicsSetFrameTimingHint(GraphicsClient, FTH);
+			if (Result == MLResult_Ok)
 			{
 				UE_LOG(LogMagicLeap, Log, TEXT("Set Lumin frame timing hint to %s."), *FrameTimingEnum->GetNameStringByValue((int64)InFrameTimingHint));
 				CurrentFrameTimingHint = InFrameTimingHint;
@@ -721,10 +710,7 @@ void FMagicLeapHMD::SetFrameTimingHint(ELuminFrameTimingHint InFrameTimingHint)
 
 float FMagicLeapHMD::GetWorldToMetersScale() const
 {
-	FTrackingFrame* TrackingFrame = GetCurrentFrame();
-	check(TrackingFrame);
-
-	return TrackingFrame->WorldToMetersScale;
+	return GetCurrentFrame().WorldToMetersScale;
 }
 
 bool FMagicLeapHMD::EnableStereo(bool bStereo)
@@ -802,9 +788,9 @@ FMatrix FMagicLeapHMD::GetStereoProjectionMatrix(const enum EStereoscopicPass St
 	check(IsInGameThread());
 	check(IsStereoEnabled());
 	const int viewport = (StereoPassType == eSSP_LEFT_EYE) ? 0 : 1;
-	FTrackingFrame* frame = GetCurrentFrame();
+	const FTrackingFrame& frame = GetCurrentFrame();
 	// TODO: Remove this for vulkan when we can get a better result from the frame
-	return (frame && bDeviceInitialized && !IsVulkanPlatform(GMaxRHIShaderPlatform)) ? MagicLeap::ToFMatrix(frame->UpdateInfoArray.virtual_camera_extents[viewport].projection) : FMatrix::Identity;
+	return (bDeviceInitialized && !IsVulkanPlatform(GMaxRHIShaderPlatform)) ? MagicLeap::ToFMatrix(frame.UpdateInfoArray.virtual_camera_extents[viewport].projection) : FMatrix::Identity;
 #else
 	return FMatrix();
 #endif //WITH_MLSDK
@@ -1153,6 +1139,7 @@ void FMagicLeapHMD::EnableDeviceFeatures()
 	// In this case, just skip these steps since their timeouts may cause the game to appear to hang.
 	if (IsHMDConnected())
 	{
+		EnablePrivileges();
 		EnablePerception();
 		EnableHeadTracking();
 		EnableInputDevices();
@@ -1167,6 +1154,7 @@ void FMagicLeapHMD::DisableDeviceFeatures()
 	DisableInputDevices();
 	DisableHeadTracking();
 	DisablePerception();
+	DisablePrivileges();
 	if (GIsEditor)
 	{
 		ReleaseDevice();
@@ -1180,7 +1168,7 @@ void FMagicLeapHMD::InitDevice_RenderThread()
 #if WITH_MLSDK
 	if (!bDeviceInitialized)
 	{
-		int32 bDeviceSuccessfullyInitialized = 0;
+		bool bDeviceSuccessfullyInitialized = false;
 		// Unreal supports sRGB which is the default we are requesting from graphics as well now.
 		MLGraphicsOptions gfx_opts;
 		gfx_opts.graphics_flags = 0; //MLGraphicsFlags_DebugMode;
@@ -1198,7 +1186,7 @@ void FMagicLeapHMD::InitDevice_RenderThread()
 #if PLATFORM_WINDOWS
 		if (IsPCPlatform(GMaxRHIShaderPlatform) && !IsOpenGLPlatform(GMaxRHIShaderPlatform))
 		{
-			bDeviceSuccessfullyInitialized = 1;
+			bDeviceSuccessfullyInitialized = true;
 			//UE_LOG(LogMagicLeap, Error, TEXT("FMagicLeapCustomPresentD3D11 is not supported."));
 		}
 #endif // PLATFORM_WINDOWS
@@ -1206,7 +1194,7 @@ void FMagicLeapHMD::InitDevice_RenderThread()
 #if PLATFORM_MAC
 		if (IsMetalPlatform(GMaxRHIShaderPlatform) && !IsOpenGLPlatform(GMaxRHIShaderPlatform))
 		{
-			bDeviceSuccessfullyInitialized = 1;
+			bDeviceSuccessfullyInitialized = true;
 			//UE_LOG(LogMagicLeap, Error, TEXT("FMagicLeapCustomPresentMetal is not supported."));
 		}
 #endif // PLATFORM_MAC
@@ -1218,16 +1206,17 @@ void FMagicLeapHMD::InitDevice_RenderThread()
 			MLHandle ContextHandle;
 			auto OpenGLRHI = static_cast<FOpenGLDynamicRHI*>(GDynamicRHI);
 			ContextHandle = reinterpret_cast<MLHandle>(OpenGLRHI->GetOpenGLCurrentContextHandle());
-			MLStatus OutStatus;
-			bDeviceSuccessfullyInitialized = MLGraphicsCreateClientGL(&gfx_opts, ContextHandle, &GraphicsClient, &OutStatus);
-			if (bDeviceSuccessfullyInitialized)
+			MLResult Result = MLGraphicsCreateClientGL(&gfx_opts, ContextHandle, &GraphicsClient);
+			if (Result == MLResult_Ok)
 			{
+				bDeviceSuccessfullyInitialized = true;
 				InitializeClipExtents_RenderThread();
 			}
 			else
 			{
+				bDeviceSuccessfullyInitialized = false;
 				GraphicsClient = ML_INVALID_HANDLE;
-				UE_LOG(LogMagicLeap, Error, TEXT("MLGraphicsCreateClientGL failed with status %d"), static_cast<int32>(OutStatus));
+				UE_LOG(LogMagicLeap, Error, TEXT("MLGraphicsCreateClientGL failed with status %d"), Result);
 			}
 		}
 		// @todo Lumin vulkan: THis doesn't compile, even with some #includes - maybe need more, I dunno
@@ -1254,24 +1243,24 @@ void FMagicLeapHMD::InitDevice_RenderThread()
 			{
 				UE_LOG(LogMagicLeap, Display, TEXT("FMagicLeapCustomPresentVulkan is supported."));
 				FVulkanDynamicRHI* VulkanDynamicRHI = (FVulkanDynamicRHI*)GDynamicRHI;
+				uint64 Instance = VulkanRHIBridge::GetInstance(VulkanDynamicRHI);
 				FVulkanDevice* VulkanDevice = VulkanRHIBridge::GetDevice(VulkanDynamicRHI);
 				uint64 PhysicalDevice = VulkanRHIBridge::GetPhysicalDevice(VulkanDevice);
 				uint64 LogicalDevice = VulkanRHIBridge::GetLogicalDevice(VulkanDevice);
 				GraphicsClient = ML_INVALID_HANDLE;
-				MLStatus OutStatus;
-				bool bDeviceSuccessfullyInitialized2 = MLGraphicsCreateClientVk(&gfx_opts, (VkPhysicalDevice)PhysicalDevice, (VkDevice)LogicalDevice, &GraphicsClient, &OutStatus);
-				if (bDeviceSuccessfullyInitialized2)
+				MLResult Result = MLGraphicsCreateClientVk(&gfx_opts, (VkInstance)Instance, (VkPhysicalDevice)PhysicalDevice, (VkDevice)LogicalDevice, &GraphicsClient);
+				if (Result == MLResult_Ok)
 				{
 					InitializeClipExtents_RenderThread();
 				}
 				else
 				{
 					GraphicsClient = ML_INVALID_HANDLE;
-					UE_LOG(LogMagicLeap, Error, TEXT("MLGraphicsCreateClientVk failed with status %d"), static_cast<int32>(OutStatus));
+					UE_LOG(LogMagicLeap, Error, TEXT("MLGraphicsCreateClientVk failed with status %d"), Result);
 				}
 
-				FPlatformAtomics::InterlockedExchange(&bDeviceInitialized, bDeviceSuccessfullyInitialized2);
-				FPlatformAtomics::InterlockedExchange(&bDeviceWasJustInitialized, bDeviceSuccessfullyInitialized2);
+				FPlatformAtomics::InterlockedExchange(&bDeviceInitialized, Result == MLResult_Ok);
+				FPlatformAtomics::InterlockedExchange(&bDeviceWasJustInitialized, Result == MLResult_Ok);
 			});
 #endif // PLATFORM_LUMIN
 		}
@@ -1422,11 +1411,10 @@ void FMagicLeapHMD::ReleaseDevice_RenderThread()
 #endif
 
 #if WITH_MLSDK
-		MLStatus OutStatus;
-		bool bResult = MLGraphicsDestroyClient(&GraphicsClient, &OutStatus);
-		if (!bResult)
+		MLResult Result = MLGraphicsDestroyClient(&GraphicsClient);
+		if (Result != MLResult_Ok)
 		{
-			UE_LOG(LogMagicLeap, Error, TEXT("MLGraphicsDestroyClient failed with status %d"), static_cast<int32>(OutStatus));
+			UE_LOG(LogMagicLeap, Error, TEXT("MLGraphicsDestroyClient failed with status %d"), Result);
 		}
 		else
 		{
@@ -1438,16 +1426,11 @@ void FMagicLeapHMD::ReleaseDevice_RenderThread()
 
 bool FMagicLeapHMD::GetCurrentPose(int32 DeviceId, FQuat& OutOrientation, FVector& OutPosition)
 {
-	FTrackingFrame* Frame = GetCurrentFrame();
-	if (Frame)
-	{
-		OutOrientation = Frame->RawPose.GetRotation();
-		OutPosition = Frame->RawPose.GetLocation();
+	const FTrackingFrame& Frame = GetCurrentFrame();
+	OutOrientation = Frame.RawPose.GetRotation();
+	OutPosition = Frame.RawPose.GetLocation();
 
-		return true;
-	}
-
-	return false;
+	return true;
 }
 
 bool FMagicLeapHMD::GetRelativeEyePose(int32 DeviceId, EStereoscopicPass Eye, FQuat& OutOrientation, FVector& OutPosition)
@@ -1457,11 +1440,11 @@ bool FMagicLeapHMD::GetRelativeEyePose(int32 DeviceId, EStereoscopicPass Eye, FQ
 	OutPosition = FVector::ZeroVector;
 	if (DeviceId == IXRTrackingSystem::HMDDeviceId && (Eye == eSSP_LEFT_EYE || Eye == eSSP_RIGHT_EYE))
 	{
-		FTrackingFrame const* const Frame = GetCurrentFrame();
+		const FTrackingFrame& Frame = GetCurrentFrame();
 		const int EyeIdx = (Eye == eSSP_LEFT_EYE) ? 0 : 1;
 
-		const FTransform EyeToWorld = MagicLeap::ToFTransform(Frame->RenderInfoArray.virtual_cameras[EyeIdx].transform, Frame->WorldToMetersScale);		// "world" here means the HMDs tracking space
-		const FTransform EyeToHMD = EyeToWorld * Frame->RawPose.Inverse();		// RawPose is HMDToWorld
+		const FTransform EyeToWorld = MagicLeap::ToFTransform(Frame.RenderInfoArray.virtual_cameras[EyeIdx].transform, Frame.WorldToMetersScale);		// "world" here means the HMDs tracking space
+		const FTransform EyeToHMD = EyeToWorld * Frame.RawPose.Inverse();		// RawPose is HMDToWorld
 		OutPosition = EyeToHMD.GetTranslation();
 		OutOrientation = EyeToHMD.GetRotation();
 
@@ -1500,6 +1483,11 @@ void FMagicLeapHMD::OnBeginRendering_GameThread()
 	check(IsInGameThread());
 
 	RefreshTrackingFrame();
+	FTrackingFrame& TrackingFrame = GetCurrentFrameMutable();
+	if (TrackingFrame.WorldContext != nullptr)
+	{
+		RefreshTrackingToWorldTransform(*(TrackingFrame.WorldContext));
+	}
 
 #if WITH_MLSDK
 	// Copy the game tracking frame to the render frame.
@@ -1638,13 +1626,10 @@ void FMagicLeapHMD::GetOrthoProjection(int32 RTWidth, int32 RTHeight, float Orth
 void FMagicLeapHMD::SetClippingPlanes(float NCP, float FCP)
 {
 	check(IsInGameThread());
-	FTrackingFrame* frame = GetCurrentFrame();
-	if (frame != nullptr)
-	{
-		frame->FarClippingPlane = (frame->RecommendedFarClippingPlane < FCP) ? frame->RecommendedFarClippingPlane : FCP;
-		GNearClippingPlane = NCP;
-		UpdateNearClippingPlane();
-	}
+	FTrackingFrame& frame = GetCurrentFrameMutable();
+	frame.FarClippingPlane = (frame.RecommendedFarClippingPlane < FCP) ? frame.RecommendedFarClippingPlane : FCP;
+	GNearClippingPlane = NCP;
+	UpdateNearClippingPlane();
 }
 
 void FMagicLeapHMD::RegisterMagicLeapInputDevice(IMagicLeapInputDevice* InputDevice)
@@ -1701,26 +1686,42 @@ void FMagicLeapHMD::ShutdownRendering()
 #endif // PLATFORM_LUMIN
 }
 
-FTrackingFrame* FMagicLeapHMD::GetCurrentFrame() const
+FTrackingFrame& FMagicLeapHMD::GetCurrentFrameMutable()
 {
 	if (IsInRHIThread())
 	{
-		return const_cast<FTrackingFrame*>(&RHITrackingFrame);
+		return RHITrackingFrame;
 	}
 	if (IsInRenderingThread())
 	{
-		return const_cast<FTrackingFrame*>(&RenderTrackingFrame);
+		return RenderTrackingFrame;
 	}
 	else
 	{
-		return const_cast<FTrackingFrame*>(&GameTrackingFrame);
+		return GameTrackingFrame;
 	}
 }
 
-FTrackingFrame* FMagicLeapHMD::GetOldFrame() const
+const FTrackingFrame& FMagicLeapHMD::GetCurrentFrame() const
+{
+	if (IsInRHIThread())
+	{
+		return RHITrackingFrame;
+	}
+	if (IsInRenderingThread())
+	{
+		return RenderTrackingFrame;
+	}
+	else
+	{
+		return GameTrackingFrame;
+	}
+}
+
+const FTrackingFrame& FMagicLeapHMD::GetOldFrame() const
 {
 	check(IsInGameThread());
-	return const_cast<FTrackingFrame*>(&OldTrackingFrame);
+	return OldTrackingFrame;
 }
 
 void FMagicLeapHMD::InitializeOldFrameFromRenderFrame()
@@ -1845,6 +1846,26 @@ void FMagicLeapHMD::RestoreBaseProfile()
 	}
 }
 
+void FMagicLeapHMD::EnablePrivileges()
+{
+#if WITH_MLSDK
+	UE_LOG(LogMagicLeap, Warning, TEXT("FMagicLeapHMD::EnablePrivileges"));
+	MLResult Result = MLPrivilegesStartup();
+	UE_CLOG(Result != MLResult_Ok, LogMagicLeap, Error, TEXT("MLPrivilegesStartup() "
+		"failed with error %s"), UTF8_TO_TCHAR(MLPrivilegesGetResultString(Result)));
+#endif // WITH_MLSDK
+}
+
+void FMagicLeapHMD::DisablePrivileges()
+{
+#if WITH_MLSDK
+	UE_LOG(LogMagicLeap, Warning, TEXT("FMagicLeapHMD::DisablePrivileges"));
+	MLResult Result = MLPrivilegesShutdown();
+	UE_CLOG(Result != MLResult_Ok, LogMagicLeap, Error, TEXT("MLPrivilegesShutdown() "
+		"failed with error %s"), UTF8_TO_TCHAR(MLPrivilegesGetResultString(Result)));
+#endif // WITH_MLSDK
+}
+
 void FMagicLeapHMD::EnableInputDevices()
 {
 	for (auto& It : InputDevices)
@@ -1870,11 +1891,22 @@ void FMagicLeapHMD::EnablePerception()
 	if (!bIsPerceptionEnabled)
 	{
 		MLPerceptionSettings perception_settings;
-		MLPerceptionInitSettings(&perception_settings);
-		bIsPerceptionEnabled = MLPerceptionStartup(&perception_settings);
-		if (!bIsPerceptionEnabled)
+		MLResult Result = MLPerceptionInitSettings(&perception_settings);
+		if (Result == MLResult_Ok)
 		{
-			UE_LOG(LogMagicLeap, Error, TEXT("Error starting perception client."));
+			Result = MLPerceptionStartup(&perception_settings);
+			if (Result == MLResult_Ok)
+			{
+				bIsPerceptionEnabled = true;
+			}
+			else
+			{
+				UE_LOG(LogMagicLeap, Error, TEXT("MLPerceptionStartup failed with error %d."), Result);
+			}
+		}
+		else
+		{
+			UE_LOG(LogMagicLeap, Error, TEXT("MLPerceptionInitSettings failed with error %d."), Result);
 		}
 	}
 #endif //WITH_MLSDK
@@ -1885,14 +1917,15 @@ void FMagicLeapHMD::DisablePerception()
 #if WITH_MLSDK
 	if (bIsPerceptionEnabled)
 	{
-		bIsPerceptionEnabled = !MLPerceptionShutdown();
-		if (bIsPerceptionEnabled)
+		MLResult Result = MLPerceptionShutdown();
+		if (Result == MLResult_Ok)
 		{
-			UE_LOG(LogMagicLeap, Error, TEXT("Error shutting down perception client."));
+			bIsPerceptionEnabled = false;
+			UE_LOG(LogMagicLeap, Display, TEXT("Perception client shutdown successfully."));
 		}
 		else
 		{
-			UE_LOG(LogMagicLeap, Display, TEXT("Perception client shutdown successfully."));
+			UE_LOG(LogMagicLeap, Error, TEXT("MLPerceptionShutdown failed with error %d"), Result);
 		}
 	}
 #endif //WITH_MLSDK
@@ -1903,10 +1936,13 @@ void FMagicLeapHMD::EnableHeadTracking()
 #if WITH_MLSDK
 	if (HeadTracker == ML_INVALID_HANDLE)
 	{
-		HeadTracker = MLHeadTrackingCreate();
-		if (HeadTracker != ML_INVALID_HANDLE)
+		MLResult CreateResult = MLHeadTrackingCreate(&HeadTracker);
+		if (CreateResult == MLResult_Ok && HeadTracker != ML_INVALID_HANDLE)
 		{
-			MLHeadTrackingGetStaticData(HeadTracker, &HeadTrackerData);
+			if (MLResult_Ok != MLHeadTrackingGetStaticData(HeadTracker, &HeadTrackerData))
+			{
+				UE_LOG(LogMagicLeap, Error, TEXT("Error getting head tracking static data."));
+			}
 		}
 		else
 		{
@@ -1921,15 +1957,9 @@ void FMagicLeapHMD::DisableHeadTracking()
 #if WITH_MLSDK
 	if (HeadTracker != ML_INVALID_HANDLE)
 	{
-		bool bResult = MLHeadTrackingDestroy(HeadTracker);
-		if (bResult)
-		{
-			HeadTracker = ML_INVALID_HANDLE;
-		}
-		else
-		{
-			UE_LOG(LogMagicLeap, Error, TEXT("Error destroying head tracking."));
-		}
+		MLResult DestroyResult = MLHeadTrackingDestroy(HeadTracker);
+		UE_CLOG(DestroyResult != MLResult_Ok, LogMagicLeap, Error, TEXT("Error destroying head tracking."));
+		HeadTracker = ML_INVALID_HANDLE;
 	}
 #endif //WITH_MLSDK
 }
@@ -1937,10 +1967,9 @@ void FMagicLeapHMD::DisableHeadTracking()
 void FMagicLeapHMD::InitializeClipExtents_RenderThread()
 {
 #if WITH_MLSDK
-	MLStatus OutStatus;
 	MLGraphicsRenderTargetsInfo RenderTargetInfo;
-	bool bResult = MLGraphicsGetRenderTargets(GraphicsClient, &RenderTargetInfo, &OutStatus);
-	if (bResult)
+	MLResult Result = MLGraphicsGetRenderTargets(GraphicsClient, &RenderTargetInfo);
+	if (Result == MLResult_Ok)
 	{
 		GameTrackingFrame.NearClippingPlane = RenderTargetInfo.min_clip * GameTrackingFrame.WorldToMetersScale;
 		GameTrackingFrame.RecommendedFarClippingPlane = RenderTargetInfo.max_clip * GameTrackingFrame.WorldToMetersScale;
@@ -1948,14 +1977,14 @@ void FMagicLeapHMD::InitializeClipExtents_RenderThread()
 	}
 	else
 	{
-		UE_LOG(LogMagicLeap, Error, TEXT("MLGraphicsGetRenderTargets() failed with status %d"), static_cast<int32>(OutStatus));
+		UE_LOG(LogMagicLeap, Error, TEXT("MLGraphicsGetRenderTargets() failed with status %d"), Result);
 	}
 
 	// get the clip extents for clipping content in update thread
-	bResult = MLGraphicsGetClipExtents(GraphicsClient, &GameTrackingFrame.UpdateInfoArray, &OutStatus);
-	if (!bResult)
+	Result = MLGraphicsGetClipExtents(GraphicsClient, &GameTrackingFrame.UpdateInfoArray);
+	if (Result != MLResult_Ok)
 	{
-		FString ErrorMesg = FString::Printf(TEXT("MLGraphicsGetClipExtents failed with status %d"), static_cast<int32>(OutStatus));
+		FString ErrorMesg = FString::Printf(TEXT("MLGraphicsGetClipExtents failed with status %d"), Result);
 
 		// In case we're running under VD/ZI, there's always the risk of disconnects.
 		// In those cases, the graphics API can return an error, but the client handle might still be valid.
diff --git a/Engine/Plugins/Lumin/MagicLeap/Source/MagicLeap/Private/MagicLeapHMD.h b/Engine/Plugins/Lumin/MagicLeap/Source/MagicLeap/Private/MagicLeapHMD.h
--- a/Engine/Plugins/Lumin/MagicLeap/Source/MagicLeap/Private/MagicLeapHMD.h
+++ b/Engine/Plugins/Lumin/MagicLeap/Source/MagicLeap/Private/MagicLeapHMD.h
@@ -9,6 +9,7 @@
 #include "Misc/ScopeLock.h"
 
 #include "XRRenderTargetManager.h"
+#include "IStereoLayers.h"
 #include "AppFramework.h"
 #include "SceneViewExtension.h"
 #include "MagicLeapMath.h"
@@ -172,9 +173,11 @@ public:
 public:
 
 	uint32 GetViewportCount() const { return AppFramework.IsInitialized() ? AppFramework.GetViewportCount() : 0; }
-	FTrackingFrame* GetCurrentFrame() const;
-	FTrackingFrame* GetOldFrame() const;
 
+	// TODO: add const versions
+	FTrackingFrame& GetCurrentFrameMutable();
+	const FTrackingFrame& GetCurrentFrame() const;
+	const FTrackingFrame& GetOldFrame() const;
 
 	// HACK: This is a hack in order to pass variables from game frame to render frame
 	// This should be removed once graphics provides vergence based focus distance
@@ -227,6 +230,9 @@ private:
 	void EnableLuminProfile();
 	void RestoreBaseProfile();
 
+	void EnablePrivileges();
+	void DisablePrivileges();
+
 	void EnableInputDevices();
 	void DisableInputDevices();
 
diff --git a/Engine/Plugins/Lumin/MagicLeap/Source/MagicLeap/Private/MagicLeapPrivileges.cpp b/Engine/Plugins/Lumin/MagicLeap/Source/MagicLeap/Private/MagicLeapPrivileges.cpp
new file mode 100644
--- /dev/null
+++ b/Engine/Plugins/Lumin/MagicLeap/Source/MagicLeap/Private/MagicLeapPrivileges.cpp
@@ -0,0 +1,205 @@
+// %BANNER_BEGIN%
+// ---------------------------------------------------------------------
+// %COPYRIGHT_BEGIN%
+//
+// Copyright (c) 2017 Magic Leap, Inc. (COMPANY) All Rights Reserved.
+// Magic Leap, Inc. Confidential and Proprietary
+//
+// NOTICE: All information contained herein is, and remains the property
+// of COMPANY. The intellectual and technical concepts contained herein
+// are proprietary to COMPANY and may be covered by U.S. and Foreign
+// Patents, patents in process, and are protected by trade secret or
+// copyright law. Dissemination of this information or reproduction of
+// this material is strictly forbidden unless prior written permission is
+// obtained from COMPANY. Access to the source code contained herein is
+// hereby forbidden to anyone except current COMPANY employees, managers
+// or contractors who have executed Confidentiality and Non-disclosure
+// agreements explicitly covering such access.
+//
+// The copyright notice above does not evidence any actual or intended
+// publication or disclosure of this source code, which includes
+// information that is confidential and/or proprietary, and is a trade
+// secret, of COMPANY. ANY REPRODUCTION, MODIFICATION, DISTRIBUTION,
+// PUBLIC PERFORMANCE, OR PUBLIC DISPLAY OF OR THROUGH USE OF THIS
+// SOURCE CODE WITHOUT THE EXPRESS WRITTEN CONSENT OF COMPANY IS
+// STRICTLY PROHIBITED, AND IN VIOLATION OF APPLICABLE LAWS AND
+// INTERNATIONAL TREATIES. THE RECEIPT OR POSSESSION OF THIS SOURCE
+// CODE AND/OR RELATED INFORMATION DOES NOT CONVEY OR IMPLY ANY RIGHTS
+// TO REPRODUCE, DISCLOSE OR DISTRIBUTE ITS CONTENTS, OR TO MANUFACTURE,
+// USE, OR SELL ANYTHING THAT IT MAY DESCRIBE, IN WHOLE OR IN PART.
+//
+// %COPYRIGHT_END%
+// --------------------------------------------------------------------
+// %BANNER_END%
+
+#include "MagicLeapPrivileges.h"
+#include "Modules/ModuleManager.h"
+#include "HAL/PlatformProcess.h"
+#include "Misc/Paths.h"
+#include "MagicLeapPluginUtil.h"
+#include "MagicLeapSDKDetection.h"
+#include "AppFramework.h"
+
+#if WITH_MLSDK
+#include "ml_privileges.h"
+#endif //WITH_MLSDK
+
+//////////////////////////////////////////////////////////////////////////
+
+#if WITH_MLSDK
+
+MLPrivilegeID UnrealToMLPrivilege(EMagicLeapPrivilege Privilege)
+{
+// TODO: We need to get rid of any hand-mapping of these enums. In the meantime,
+// the macro is to make it easier to keep it in step with the header - rmobbs
+#define PRIVCASE(x) case EMagicLeapPrivilege::x: { return MLPrivilegeID_##x; }
+	switch(Privilege)
+	{
+		PRIVCASE(AudioRecognizer)
+		PRIVCASE(BatteryInfo)
+		PRIVCASE(CameraCapture)
+		PRIVCASE(WorldReconstruction)
+		PRIVCASE(InAppPurchase)
+		PRIVCASE(AudioCaptureMic)
+		PRIVCASE(DrmCertificates)
+		PRIVCASE(Occlusion)
+		PRIVCASE(LowLatencyLightwear)
+		PRIVCASE(Internet)
+		PRIVCASE(IdentityRead)
+		PRIVCASE(BackgroundDownload)
+		PRIVCASE(BackgroundUpload)
+		PRIVCASE(MediaDrm)
+		PRIVCASE(Media)
+		PRIVCASE(MediaMetadata)
+		PRIVCASE(PowerInfo)
+		PRIVCASE(LocalAreaNetwork)
+		PRIVCASE(VoiceInput)
+		PRIVCASE(Documents)
+		PRIVCASE(ConnectBackgroundMusicService)
+		PRIVCASE(RegisterBackgroundMusicService)
+		PRIVCASE(PwFoundObjRead)
+		PRIVCASE(NormalNotificationsUsage)
+		PRIVCASE(MusicService)
+		PRIVCASE(ControllerPose)
+		PRIVCASE(ScreensProvider)
+		PRIVCASE(GesturesSubscribe)
+		PRIVCASE(GesturesConfig)
+	default:
+		UE_LOG(LogMagicLeap, Error, TEXT("Unmapped privilege %d"), static_cast<int32>(Privilege));
+		return MLPrivilegeID_Invalid;
+	}
+}
+#endif //WITH_MLSDK
+
+UMagicLeapPrivileges::UMagicLeapPrivileges()
+{
+	// Make sure this component ticks
+	PrimaryComponentTick.bCanEverTick = true;
+	PrimaryComponentTick.bStartWithTickEnabled = true;
+	PrimaryComponentTick.TickGroup = TG_PrePhysics;
+	bAutoActivate = true;
+}
+
+UMagicLeapPrivileges::~UMagicLeapPrivileges()
+{
+}
+
+bool UMagicLeapPrivileges::CheckPrivilege(EMagicLeapPrivilege Privilege)
+{
+#if WITH_MLSDK
+	MLResult Result = MLPrivilegesCheckPrivilege(UnrealToMLPrivilege(Privilege));
+
+	UE_LOG(LogMagicLeap, Warning, TEXT("UMagicLeapPrivileges::CheckPrivilege got result "
+		"%u %s for privilege #%d"), static_cast<uint32>(Result), 
+		UTF8_TO_TCHAR(MLPrivilegesGetResultString(Result)), 
+		static_cast<int32>(UnrealToMLPrivilege(Privilege)));
+
+	if (MLPrivilegesResult_Granted == Result)
+	{
+		return true;
+	}
+#endif //WITH_MLSDK
+	return false;
+}
+
+bool UMagicLeapPrivileges::RequestPrivilege(EMagicLeapPrivilege Privilege)
+{
+#if WITH_MLSDK
+	MLResult Result = MLPrivilegesRequestPrivilege(UnrealToMLPrivilege(Privilege));
+
+	UE_LOG(LogMagicLeap, Warning, TEXT("UMagicLeapPrivileges::RequestPrivilege got result "
+		"%u %s for privilege #%d"), static_cast<uint32>(Result), 
+		UTF8_TO_TCHAR(MLPrivilegesGetResultString(Result)), 
+		static_cast<int32>(UnrealToMLPrivilege(Privilege)));
+
+	if (MLPrivilegesResult_Granted == Result)
+	{
+		return true;
+	}
+#endif //WITH_MLSDK
+	return false;
+}
+
+bool UMagicLeapPrivileges::RequestPrivilegeAsync(EMagicLeapPrivilege Privilege, const FPrivilegeRequestDelegate& ResultDelegate)
+{
+#if WITH_MLSDK
+	MLPrivilegesAsyncRequest* AsyncPrivilegeRequest;
+
+	auto Result = MLPrivilegesRequestPrivilegeAsync(UnrealToMLPrivilege(Privilege), &AsyncPrivilegeRequest);
+
+	UE_LOG(LogMagicLeap, Warning, TEXT("UMagicLeapPrivileges::RequestPrivilegeAsync got result "
+		"%u %s for privilege #%d"), static_cast<uint32>(Result), 
+		UTF8_TO_TCHAR(MLPrivilegesGetResultString(Result)), 
+		static_cast<int32>(UnrealToMLPrivilege(Privilege)));
+
+	if (MLResult_Ok != Result)
+	{
+		UE_LOG(LogMagicLeap, Error, TEXT("UMagicLeapPrivileges::RequestPrivilegeAsync "
+			"failure %u %s for privilege #%d"), static_cast<uint32>(Result), 
+			UTF8_TO_TCHAR(MLPrivilegesGetResultString(Result)),
+			static_cast<int32>(UnrealToMLPrivilege(Privilege)));
+		return false;
+	}
+
+	// Store the request
+	PendingAsyncRequests.Add({ Privilege, AsyncPrivilegeRequest, ResultDelegate });
+	return true;
+#else
+	return false;
+#endif //WITH_MLSDK
+}
+
+void UMagicLeapPrivileges::TickComponent(float DeltaTime, enum ELevelTick TickType, FActorComponentTickFunction *ThisTickFunction)
+{
+#if WITH_MLSDK
+	auto CopyPendingAsyncRequests(PendingAsyncRequests);
+	
+	PendingAsyncRequests.Empty();
+
+	for (const auto& PendingAsyncRequest : CopyPendingAsyncRequests)
+	{
+		// Re-add pending ones
+		MLResult Result = MLPrivilegesRequestPrivilegeTryGet(PendingAsyncRequest.Request);
+		if (MLResult_Pending == Result)
+		{
+			PendingAsyncRequests.Add(PendingAsyncRequest);
+			continue;
+		}
+
+		UE_LOG(LogMagicLeap, Warning, TEXT("UMagicLeapPrivileges::"
+			"TickComponent has result %u %s for privilege request %d"), 
+			static_cast<uint32>(Result), UTF8_TO_TCHAR(MLPrivilegesGetResultString(Result)), 
+			static_cast<int32>(PendingAsyncRequest.Privilege));
+
+		// Dispatch. Granted gets true, everything else gets false.
+		if (MLPrivilegesResult_Granted == Result)
+		{
+			PendingAsyncRequest.Delegate.ExecuteIfBound(PendingAsyncRequest.Privilege, true);
+		}
+		else
+		{
+			PendingAsyncRequest.Delegate.ExecuteIfBound(PendingAsyncRequest.Privilege, false);
+		}
+	}
+#endif
+}
diff --git a/Engine/Plugins/Lumin/MagicLeap/Source/MagicLeap/Private/MagicLeapXRCamera.cpp b/Engine/Plugins/Lumin/MagicLeap/Source/MagicLeap/Private/MagicLeapXRCamera.cpp
--- a/Engine/Plugins/Lumin/MagicLeap/Source/MagicLeap/Private/MagicLeapXRCamera.cpp
+++ b/Engine/Plugins/Lumin/MagicLeap/Source/MagicLeap/Private/MagicLeapXRCamera.cpp
@@ -19,11 +19,11 @@ void FMagicLeapXRCamera::PreRenderView_RenderThread(FRHICommandListImmediate& RH
 	if (View.StereoPass != eSSP_FULL)
 	{
 		const int EyeIdx = (View.StereoPass == eSSP_LEFT_EYE) ? 0 : 1;
-		FTrackingFrame* const Frame = MagicLeapSystem.GetCurrentFrame();
+		const FTrackingFrame& Frame = MagicLeapSystem.GetCurrentFrame();
 
 		// update to use render projection matrix
 		// #todo: Roll UpdateProjectionMatrix into UpdateViewMatrix?
-		FMatrix RenderInfoProjectionMatrix = MagicLeap::ToFMatrix(Frame->RenderInfoArray.virtual_cameras[EyeIdx].projection);
+		FMatrix RenderInfoProjectionMatrix = MagicLeap::ToFMatrix(Frame.RenderInfoArray.virtual_cameras[EyeIdx].projection);
 
 		// Set the near clipping plane to GNearClippingPlane which is clamped to the minimum value allowed for the device. (ref: MLGraphicsGetRenderTargets())
 		RenderInfoProjectionMatrix.M[3][2] = GNearClippingPlane;
