#DEV-VR - Merge from //UE4/Dev-MagicLeap/... @ CL 4136411

diff --git a/Engine/Plugins/Lumin/MagicLeap/Source/MagicLeap/Private/MeshTrackerComponent.cpp b/Engine/Plugins/Lumin/MagicLeap/Source/MagicLeap/Private/MeshTrackerComponent.cpp
--- a/Engine/Plugins/Lumin/MagicLeap/Source/MagicLeap/Private/MeshTrackerComponent.cpp
+++ b/Engine/Plugins/Lumin/MagicLeap/Source/MagicLeap/Private/MeshTrackerComponent.cpp
@@ -1,6 +1,7 @@
 // Copyright 1998-2018 Epic Games, Inc. All Rights Reserved.
 
 #include "MeshTrackerComponent.h"
+#include "IMagicLeapHMD.h"
 #include "MagicLeapHMD.h"
 #include "AppFramework.h"
 #include "MagicLeapMath.h"
@@ -19,188 +20,145 @@
 
 #if WITH_MLSDK
 ML_INCLUDES_START
-#include <ml_meshing.h>
-#include <ml_data_array.h>
+#include <ml_meshing2.h>
 ML_INCLUDES_END
 #endif //WITH_MLSDK
 
-static TAutoConsoleVariable<float> CVarFakeMeshTrackerData(
-	TEXT("vr.MagicLeap.FakeMeshTrackerData"),
-	0.0f,
-	TEXT("If True MeshTrackerComponent will generate some simple fake mesh data.\n"),
-	ECVF_Default);
+#if WITH_MLSDK
+// Map an Unreal meshing LOD to the corresponding ML meshing LOD
+FORCEINLINE MLMeshingLOD MLToUnreal_MeshLOD(EMeshLOD UnrealMeshLOD)
+{
+	switch (UnrealMeshLOD)
+	{
+		case EMeshLOD::Minimum:
+			return MLMeshingLOD_Minimum;
+		case EMeshLOD::Medium:
+			return MLMeshingLOD_Medium;
+		case EMeshLOD::Maximum:
+			return MLMeshingLOD_Maximum;
+	}
+	check(false);
+	return MLMeshingLOD_Minimum;
+}
+
+// Makes MLCoordinateFrameUID a hashable type for use with TMap and TSet
+template<typename ValueType> struct TMLCoordinateFrameUIDKeyFuncs : 
+	BaseKeyFuncs<TPair<MLCoordinateFrameUID, ValueType>, MLCoordinateFrameUID>
+{
+private:
+	typedef BaseKeyFuncs<TPair<MLCoordinateFrameUID, ValueType>, MLCoordinateFrameUID> Super;
+public:
+	typedef typename Super::ElementInitType ElementInitType;
+	typedef typename Super::KeyInitType KeyInitType;
+
+	static KeyInitType GetSetKey(ElementInitType Element)
+	{
+		return Element.Key;
+	}
+	static bool Matches(KeyInitType A, KeyInitType B)
+	{
+		return A.data[0] == B.data[0] && A.data[1] == B.data[1];
+	}
+	static uint32 GetKeyHash(KeyInitType Key)
+	{
+		return GetTypeHash(static_cast<uint64>(Key.data[0])) ^ ~GetTypeHash(static_cast<uint64>(Key.data[1]));
+	}
+};
+#endif //WITH_MLSDK
 
 class FMeshTrackerImpl : public MagicLeap::IAppEventHandler
 {
 public:
 	FMeshTrackerImpl()
 #if WITH_MLSDK
-		: Tracker(ML_INVALID_HANDLE)
-		, bUpdateRequested(false)
+		: MeshTracker(ML_INVALID_HANDLE)
+		, MeshBrickIndex(0)
+		, CurrentMeshInfoRequest(ML_INVALID_HANDLE)
+		, CurrentMeshRequest(ML_INVALID_HANDLE)
 #endif //WITH_MLSDK
 	{
-#if WITH_MLSDK
-		UnrealToMLMeshTypeMap.Add(EMeshType::Full, MLMeshingType_Full);
-		UnrealToMLMeshTypeMap.Add(EMeshType::Blocks, MLMeshingType_Blocks);
-		UnrealToMLMeshTypeMap.Add(EMeshType::PointCloud, MLMeshingType_PointCloud);
-#endif //WITH_MLSDK
 	};
 
 #if WITH_MLSDK
-	MLMeshingSettings CreateSettings(const UMeshTrackerComponent& MeshTracker)
+	MLMeshingSettings CreateSettings(const UMeshTrackerComponent& MeshTrackerComponent)
 	{
+		MLMeshingSettings Settings;
+
+		MLMeshingInitSettings(&Settings);
+
 		float WorldToMetersScale = 100.0f;
-		if (GEngine->XRSystem.IsValid() && GEngine->XRSystem->GetHMDDevice())
+		if (IMagicLeapPlugin::Get().IsMagicLeapHMDValid())
 		{
-			const FAppFramework& AppFramework = static_cast<FMagicLeapHMD*>(GEngine->XRSystem->GetHMDDevice())->GetAppFrameworkConst();
+			const FAppFramework& AppFramework = static_cast<FMagicLeapHMD *>
+				(GEngine->XRSystem->GetHMDDevice())->GetAppFrameworkConst();
 			if (AppFramework.IsInitialized())
 			{
 				WorldToMetersScale = AppFramework.GetWorldToMetersScale();
 			}
 		}
 
-		MLMeshingSettings Settings;
-		FMemory::Memset(&Settings, 0, sizeof(Settings));
-
-		// Convert all of the times from float seconds -> uint64 nanoseconds.
-		Settings.meshing_poll_time = static_cast<uint64_t>(MeshTracker.MeshingPollTime * 1e9);
-
-		Settings.mesh_type = UnrealToMLMeshTypeMap[MeshTracker.MeshType];
-
-		Settings.bounds_center = MagicLeap::ToMLVector(BoundsCenter, WorldToMetersScale);
-		Settings.bounds_rotation = MagicLeap::ToMLQuat(BoundsRotation);
-		Settings.bounds_extents = MagicLeap::ToMLVector(MeshTracker.BoundingVolume->GetScaledBoxExtent(), WorldToMetersScale);
-
-		// MagicLeap::ToMLVector() causes the Z component to be negated.
-		// The bounds were thus invalid and resulted in everything being meshed. 
-		// This provides the content devs with an option to ignore the bounding volume at will.
-		// TODO: Can this be improved?
-		if (!MeshTracker.IgnoreBoundingVolume)
+		if (MeshTrackerComponent.MeshType == EMeshType::PointCloud)
 		{
-			Settings.bounds_extents.x = FMath::Abs<float>(Settings.bounds_extents.x);
-			Settings.bounds_extents.y = FMath::Abs<float>(Settings.bounds_extents.y);
-			Settings.bounds_extents.z = FMath::Abs<float>(Settings.bounds_extents.z);
+			Settings.flags |= MLMeshingFlags_PointCloud;
+		}
+		if (MeshTrackerComponent.RequestNormals)
+		{
+			Settings.flags |= MLMeshingFlags_ComputeNormals;
+		}
+		if (MeshTrackerComponent.RequestVertexConfidence)
+		{
+			Settings.flags |= MLMeshingFlags_ComputeConfidence;
+		}
+		if (MeshTrackerComponent.Planarize)
+		{
+			Settings.flags |= MLMeshingFlags_Planarize;
+		}
+		if (MeshTrackerComponent.RemoveOverlappingTriangles)
+		{
+			Settings.flags |= MLMeshingFlags_RemoveMeshSkirt;
 		}
 
-		Settings.target_number_triangles = static_cast<uint32>(MeshTracker.TargetNumberTriangles);
-		Settings.target_number_triangles_per_block = static_cast<uint32>(MeshTracker.TargetNumberTriangles);
-
-		Settings.enable_meshing = MeshTracker.ScanWorld;
-		Settings.index_order_ccw = false;
-		Settings.fill_holes = MeshTracker.FillGaps;
-		Settings.fill_hole_length = MeshTracker.PerimeterOfGapsToFill / WorldToMetersScale;
-		Settings.compute_normals = MeshTracker.RequestNormals;
-		Settings.planarize = MeshTracker.Planarize;
-		Settings.remove_disconnected_components = MeshTracker.RemoveDisconnectedSections;
-		Settings.disconnected_component_area = MeshTracker.DisconnectedSectionArea / (WorldToMetersScale * WorldToMetersScale);
-		Settings.request_vertex_confidence = MeshTracker.RequestVertexConfidence;
-		Settings.remove_mesh_skirt = MeshTracker.RemoveOverlappingTriangles;
+		Settings.fill_hole_length = MeshTrackerComponent.PerimeterOfGapsToFill / WorldToMetersScale;
+		Settings.disconnected_component_area = MeshTrackerComponent.
+			DisconnectedSectionArea / (WorldToMetersScale * WorldToMetersScale);
 
 		return Settings;
 	};
 
-	bool SettingsChanged(const MLMeshingSettings& lhs, const MLMeshingSettings& rhs, const UMeshTrackerComponent& MeshTracker)
-	{
-		return (lhs.meshing_poll_time != rhs.meshing_poll_time ||
-			lhs.mesh_type != rhs.mesh_type ||
-			MagicLeap::ToFVector(lhs.bounds_extents, 1.0f) != MagicLeap::ToFVector(rhs.bounds_extents, 1.0f) ||
-			lhs.target_number_triangles != rhs.target_number_triangles ||
-			lhs.target_number_triangles_per_block != rhs.target_number_triangles_per_block ||
-			lhs.enable_meshing != rhs.enable_meshing ||
-			lhs.fill_holes != rhs.fill_holes ||
-			lhs.planarize != rhs.planarize ||
-			lhs.remove_disconnected_components != rhs.remove_disconnected_components ||
-			lhs.disconnected_component_area != rhs.disconnected_component_area ||
-			lhs.request_vertex_confidence != rhs.request_vertex_confidence ||
-			lhs.remove_mesh_skirt != rhs.remove_mesh_skirt ||
-			FVector(lhs.bounds_center.xyz.x - rhs.bounds_center.xyz.x,
-				lhs.bounds_center.xyz.y - rhs.bounds_center.xyz.y,
-				lhs.bounds_center.xyz.z - rhs.bounds_center.xyz.z).Size() > MeshTracker.MinDistanceRescan
-			);
-		// TODO: (njain) Account for a minimum change in bounds_rotation and bounds_extents.
-	}
 #endif //WITH_MLSDK
 
 	void OnAppPause() override
 	{
-#if WITH_MLSDK
-		bWasSystemEnabledOnPause = CurrentSettings.enable_meshing;
-
-		if (!bWasSystemEnabledOnPause)
-		{
-			UE_LOG(LogMagicLeap, Log, TEXT("Mesh tracking was not enabled at time of application pause."));
-		}
-		else
-		{
-			if (!MLHandleIsValid(Tracker))
-			{
-				UE_LOG(LogMagicLeap, Error, TEXT("Mesh tracker was invalid on application pause."));
-			}
-			else
-			{
-				CurrentSettings.enable_meshing = false;
-
-				if (!MLMeshingUpdate(Tracker, &CurrentSettings))
-				{
-					UE_LOG(LogMagicLeap, Error, TEXT("Failed to disable mesh tracker on application pause."));
-				}
-				else
-				{
-					UE_LOG(LogMagicLeap, Log, TEXT("Mesh tracker paused until app resumes."));
-				}
-			}
-		}
-#endif //WITH_MLSDK
 	}
 
 	void OnAppResume() override
 	{
-#if WITH_MLSDK
-		if (!bWasSystemEnabledOnPause)
-		{
-			UE_LOG(LogMagicLeap, Log, TEXT("Not resuming mesh tracker as it was not enabled at time of application pause."));
-		}
-		else
-		{
-			if (!MLHandleIsValid(Tracker))
-			{
-				UE_LOG(LogMagicLeap, Error, TEXT("Mesh tracker was invalid on application resume."));
-			}
-			else
-			{
-				CurrentSettings.enable_meshing = true;
-
-				if (!MLMeshingUpdate(Tracker, &CurrentSettings))
-				{
-					UE_LOG(LogMagicLeap, Error, TEXT("Failed to re-enable mesh tracker on application resume."));
-				}
-				else
-				{
-					UE_LOG(LogMagicLeap, Log, TEXT("Mesh tracker re-enabled on application resume."));
-				}
-			}
-		}
-#endif //WITH_MLSDK
 	}
 
 public:
 #if WITH_MLSDK
-	MLCoordinateFrameUID CoordinateFrame;
-	MLHandle Tracker;
-	MLMeshingStaticData Data;
+	// Handle to ML mesh tracker
+	MLHandle MeshTracker;
 
-	MLDataArrayDiff GroupDiff;
-	int32 SectionCounter;
+	// Next ID for bricks created with MR Mesh
+	int32 MeshBrickIndex;
 
-	MLMeshingSettings CurrentSettings;
+	// Handle to ML mesh info request
+	MLHandle CurrentMeshInfoRequest;
 
-	struct MeshCache
-	{
-		MLDataArrayDiff Diff;
-		int32 Section;
-	};
+	// Handle to ML mesh request
+	MLHandle CurrentMeshRequest;
+
+	// Current ML meshing settings
+	MLMeshingSettings CurrentMeshSettings;
+
+	// List of ML mesh block IDs and states
+	TArray<MLMeshingBlockRequest> MeshBlockRequests;
+
+	// Map of ML mesh block IDs to MR Mesh brick IDs
+	TMap<MLCoordinateFrameUID, uint64, FDefaultSetAllocator, 
+		TMLCoordinateFrameUIDKeyFuncs<uint64>> MeshBrickCache;
 
-	TMap<uint64, MeshCache> MeshHandleCacheMap;
 #endif //WITH_MLSDK
 
 	// Keep a copy of the mesh data here.  MRMeshComponent will use it from the game and render thread.
@@ -288,34 +246,26 @@ public:
 	FVector BoundsCenter;
 	FQuat BoundsRotation;
 
-#if WITH_MLSDK
-	bool bUpdateRequested;
-#endif //WITH_MLSDK
-
-	bool Create(const UMeshTrackerComponent& MeshTracker)
+	bool Create(const UMeshTrackerComponent& MeshTrackerComponent)
 	{
 #if WITH_MLSDK
-		if (!MLHandleIsValid(Tracker))
+		if (!MLHandleIsValid(MeshTracker))
 		{
 			// Create the tracker on demand.
-			UE_LOG(LogMagicLeap, Display, TEXT("Creating Mesh Tracker"));
-			MLMeshingSettings Settings = CreateSettings(MeshTracker);
-			Tracker = MLMeshingCreate(&Settings);
+			//UE_LOG(LogMagicLeap, Log, TEXT("Creating Mesh MeshTracker"));
+			
+			CurrentMeshSettings = CreateSettings(MeshTrackerComponent);
 
-			if (!MLHandleIsValid(Tracker))
+			MLResult Result = MLMeshingCreateClient(&MeshTracker, &CurrentMeshSettings);
+
+			if (Result != MLResult_Ok)
 			{
-				UE_LOG(LogMagicLeap, Error, TEXT("Could not create mesh tracker."));
+				UE_LOG(LogMagicLeap, Error, TEXT("MLMeshingCreateClient failed: %s."), 
+					UTF8_TO_TCHAR(MLGetResultString(Result)));
 				return false;
 			}
 
-			CurrentSettings = Settings;
-			FMemory::Memset(&GroupDiff, 0, sizeof(MLDataArrayDiff));
-			SectionCounter = 0;
-			// TODO: pull out of current scope and handle separately.
-			if (MLMeshingGetStaticData(Tracker, &Data))
-			{
-				CoordinateFrame = Data.frame;
-			}
+			MeshBrickIndex = 0;
 		}
 #endif //WITH_MLSDK
 		return true;
@@ -324,22 +274,20 @@ public:
 	void Destroy()
 	{
 #if WITH_MLSDK
-		if (MLHandleIsValid(Tracker))
+		if (MLHandleIsValid(MeshTracker))
 		{
-			bool bResult = MLMeshingDestroy(Tracker);
-			if (!bResult)
+			MLResult Result = MLMeshingDestroyClient(&MeshTracker);
+			if (Result != MLResult_Ok)
 			{
-				UE_LOG(LogMagicLeap, Error, TEXT("Error destroying mesh tracker."));
+				UE_LOG(LogMagicLeap, Error, 
+					TEXT("MLMeshingDestroyClient failed: %s."), UTF8_TO_TCHAR(MLGetResultString(Result)));
 			}
-			Tracker = ML_INVALID_HANDLE;
+			MeshTracker = ML_INVALID_HANDLE;
 		}
 #endif //WITH_MLSDK
 	}
 
 private:
-#if WITH_MLSDK
-	TMap<EMeshType, MLMeshingType> UnrealToMLMeshTypeMap;
-#endif //WITH_MLSDK
 
 	// A free list to recycle the CachedMeshData instances.  
 	TArray<FMLCachedMeshData::SharedPtr> CachedMeshDatas;
@@ -368,6 +316,7 @@ UMeshTrackerComponent::UMeshTrackerComponent(const FObjectInitializer& ObjectIni
 	BoundingVolume->CanCharacterStepUpOn = ECanBeCharacterBase::ECB_No;
 	BoundingVolume->SetCollisionObjectType(ECollisionChannel::ECC_WorldDynamic);
 	BoundingVolume->bGenerateOverlapEvents = false;
+
 	// Recommended default box extents for meshing - 10m (5m radius)
 	BoundingVolume->SetBoxExtent(FVector(1000, 1000, 1000), false);
 
@@ -395,12 +344,14 @@ void UMeshTrackerComponent::ConnectMRMesh(UMRMeshComponent* InMRMeshPtr)
 {
 	if (MRMesh)
 	{
-		UE_LOG(LogMagicLeap, Warning, TEXT("MeshTrackerComponent already has a MRMesh connected.  Ignoring this connect."));
+		UE_LOG(LogMagicLeap, Warning, 
+			TEXT("MeshTrackerComponent already has a MRMesh connected.  Ignoring this connect."));
 		return;
 	}
 	else if (InMRMeshPtr->IsConnected())
 	{
-		UE_LOG(LogMagicLeap, Warning, TEXT("MRMesh is already connected to a UMeshTrackerComponent. Ignoring this connect."));
+		UE_LOG(LogMagicLeap, Warning, 
+			TEXT("MRMesh is already connected to a UMeshTrackerComponent. Ignoring this connect."));
 		return;
 	}
 	else
@@ -414,12 +365,15 @@ void UMeshTrackerComponent::DisconnectMRMesh(class UMRMeshComponent* InMRMeshPtr
 {
 	if (!MRMesh)
 	{
-		UE_LOG(LogMagicLeap, Warning, TEXT("MeshTrackerComponent MRMesh is already disconnected. Ignoring this diconnect."));
+		UE_LOG(LogMagicLeap, Warning, 
+			TEXT("MeshTrackerComponent MRMesh is already disconnected. Ignoring this disconnect."));
 		return;
 	}
 	else if (InMRMeshPtr != MRMesh)
 	{
-		UE_LOG(LogMagicLeap, Warning, TEXT("MeshTrackerComponent MRMesh given is not the MRMesh connected. Ignoring this diconnect."));
+		UE_LOG(LogMagicLeap, Warning, 
+			TEXT("MeshTrackerComponent MRMesh given is not the MRMesh connected. "
+				 "Ignoring this disconnect."));
 		return;
 	}
 	else
@@ -434,71 +388,52 @@ void UMeshTrackerComponent::DisconnectMRMesh(class UMRMeshComponent* InMRMeshPtr
 void UMeshTrackerComponent::PostEditChangeProperty(FPropertyChangedEvent& e)
 {
 #if WITH_MLSDK
-	if (MLHandleIsValid(Impl->Tracker) && e.Property != nullptr)
+	if (MLHandleIsValid(Impl->MeshTracker) && e.Property != nullptr)
 	{
-		MLMeshingSettings Settings = Impl->CreateSettings(*this);
-		MLMeshingUpdate(Impl->Tracker, &Settings);
-	}
-#endif //WITH_MLSDK
+		MLMeshingSettings MeshSettings = Impl->CreateSettings(*this);
 
-	Super::PostEditChangeProperty(e);
-}
-#endif
+		// Just brute compare
+		if (0 != memcmp(&Impl->CurrentMeshSettings, &MeshSettings, sizeof(MeshSettings)))
+		{
+			UE_LOG(LogMagicLeap, Log, 
+				TEXT("PostEditChangeProperty is changing MLMeshingSettings"));
 
-bool UMeshTrackerComponent::ForceMeshUpdate()
-{
-#if WITH_MLSDK
-	if (MLHandleIsValid(Impl->Tracker))
-	{
-		Impl->bUpdateRequested = MLMeshingRefresh(Impl->Tracker);
-	}
-	return Impl->bUpdateRequested;
-#else
-	return true;
-#endif //WITH_MLSDK
-}
+			auto Result = MLMeshingUpdateSettings(Impl->MeshTracker, &Impl->CurrentMeshSettings);
 
-void UMeshTrackerComponent::LogMLDataArray(const MLDataArray& Data) const
-{
-#if WITH_MLSDK
-	UE_LOG(LogMagicLeap, Log, TEXT("  Data:"));
-	const uint64 timestamp = Data.timestamp;
-	const uint32 stream_count = Data.stream_count;
-	UE_LOG(LogMagicLeap, Log, TEXT("    timestamp:    %llu"), timestamp);
-	UE_LOG(LogMagicLeap, Log, TEXT("    stream_count: %i"), stream_count);
-	for (uint32 i = 0; i < stream_count; ++i)
-	{
-		const MLDataArrayStream& stream = Data.streams[i];
-		UE_LOG(LogMagicLeap, Log, TEXT("      stream: %i"), i);
-		UE_LOG(LogMagicLeap, Log, TEXT("        type: %i"), (int32)stream.type);
-		UE_LOG(LogMagicLeap, Log, TEXT("        count: %i"), stream.count);
-		UE_LOG(LogMagicLeap, Log, TEXT("        data_size: %i"), stream.data_size);
+			if (MLResult_Ok != Result)
+			{
+				UE_LOG(LogMagicLeap, Error, 
+					TEXT("MLMeshingUpdateSettings failed: %s"), UTF8_TO_TCHAR(MLGetResultString(Result)));
+			}
+			else
+			{
+				Impl->CurrentMeshSettings = MeshSettings;
+			}
+		}
 	}
 #endif //WITH_MLSDK
+
+	Super::PostEditChangeProperty(e);
 }
+#endif
 
 void UMeshTrackerComponent::TickComponent(float DeltaTime, enum ELevelTick TickType, FActorComponentTickFunction *ThisTickFunction)
 {
 	Super::TickComponent(DeltaTime, TickType, ThisTickFunction);
-#if WITH_MLSDK
 
+#if WITH_MLSDK
 	if (!MRMesh)
 	{
 		return;
 	}
 
-	static const auto FakeMeshTrackerDataCVar = IConsoleManager::Get().FindConsoleVariable(TEXT("vr.MagicLeap.FakeMeshTrackerData"));
-	if (FakeMeshTrackerDataCVar != nullptr && FakeMeshTrackerDataCVar->GetInt())
+	if (!IMagicLeapPlugin::Get().IsMagicLeapHMDValid())
 	{
-		TickWithFakeData();
 		return;
 	}
 
-	if (!GEngine->XRSystem.IsValid() || !GEngine->XRSystem->GetHMDDevice())
-	{
-		return;
-	}
-	const FAppFramework& AppFramework = static_cast<FMagicLeapHMD*>(GEngine->XRSystem->GetHMDDevice())->GetAppFrameworkConst();
+	const FAppFramework& AppFramework = 
+		static_cast<FMagicLeapHMD *>(GEngine->XRSystem->GetHMDDevice())->GetAppFrameworkConst();
 	if (!AppFramework.IsInitialized())
 	{
 		return;
@@ -515,545 +450,327 @@ void UMeshTrackerComponent::TickComponent(float DeltaTime, enum ELevelTick TickT
 	Impl->BoundsCenter = PoseInverse.TransformPosition(BoundingVolume->GetComponentLocation());
 	Impl->BoundsRotation = PoseInverse.GetRotation();
 
-	// Update the tracker on demand.
-	MLMeshingSettings NewSettings = Impl->CreateSettings(*this);
-	// Hack because MLMeshSettings does not support the comparison operator.
-	if (Impl->SettingsChanged(Impl->CurrentSettings, NewSettings, *this))
-	{
-		UE_LOG(LogMagicLeap, Display, TEXT("Updating Mesh Tracker Settings"));
-		if (Impl->CurrentSettings.mesh_type != NewSettings.mesh_type && MRMesh != nullptr)
-		{
-			MRMesh->Clear();
-		}
-		MLMeshingUpdate(Impl->Tracker, &NewSettings);
-		Impl->CurrentSettings = NewSettings;
-	}
-
-	// Update the general position of the object so the mesh is in the right position.
 	float WorldToMetersScale = AppFramework.GetWorldToMetersScale();
 
-	EFailReason FailReason = EFailReason::None;
-	FTransform Pose = FTransform::Identity;
-	if (AppFramework.GetTransform(Impl->Data.frame, Pose, FailReason))
+	// Make sure MR Mesh is at 0,0,0 (verts received from ML meshing are in world space)
+	MRMesh->SendRelativeTransform(FTransform::Identity);
+
+	if (ScanWorld)
 	{
-		if (MRMesh)
+		// Request mesh info
+		if (Impl->CurrentMeshInfoRequest == ML_INVALID_HANDLE && Impl->CurrentMeshRequest == ML_INVALID_HANDLE)
 		{
-			MRMesh->SendRelativeTransform(Pose);
+			MLMeshingExtents Extents = {};
+			Extents.center = MagicLeap::ToMLVector(Impl->BoundsCenter, WorldToMetersScale);
+			Extents.rotation = MagicLeap::ToMLQuat(Impl->BoundsRotation);
+			Extents.extents = MagicLeap::ToMLVector(BoundingVolume->GetScaledBoxExtent(), WorldToMetersScale);
+
+			// MagicLeap::ToMLVector() is, as the name implies, meant for vectors (not extents) and using it
+			// causes the Z component to be negated. We'll abs all of them for safety.
+			Extents.extents.x = FMath::Abs<float>(Extents.extents.x);
+			Extents.extents.y = FMath::Abs<float>(Extents.extents.y);
+			Extents.extents.z = FMath::Abs<float>(Extents.extents.z);
+
+			auto Result = MLMeshingRequestMeshInfo(Impl->MeshTracker, &Extents, &Impl->CurrentMeshInfoRequest);
+			if (MLResult_Ok != Result)
+			{
+				UE_LOG(LogMagicLeap, Error, TEXT("MLMeshingRequestMeshInfo failed: %s."), UTF8_TO_TCHAR(MLGetResultString(Result)));
+				Impl->CurrentMeshInfoRequest = ML_INVALID_HANDLE;
+			}
 		}
-	}
-	else if (FailReason == EFailReason::NaNsInTransform)
-	{
-		UE_LOG(LogMagicLeap, Warning, TEXT("NaNs in mesh frame transform."));
-	}
-
-	// Only draw meshes during the scan phase in order to maintain a steady framerate.
-	if (ScanWorld || Impl->bUpdateRequested)
-	{
-		// Collect current mesh handles.
-		TSet<uint64> CurrentMeshHandleSet;
 
-		// Lock group data array.
-		MLDataArray GroupData;
-		MLDataArrayLockResult mainDataArray = MLDataArrayTryLock(Impl->Data.meshes, &GroupData, &Impl->GroupDiff);
-		if (mainDataArray == MLDataArrayLockResult_New)
+		// Request block meshes for current mesh info and block list
+		if (Impl->CurrentMeshRequest == ML_INVALID_HANDLE && Impl->MeshBlockRequests.Num() > 0)
 		{
-			// Work with the locked group MLDataArray.
+			MLMeshingMeshRequest MeshRequest = {};
+			MeshRequest.request_count = static_cast<int>(Impl->MeshBlockRequests.Num());
+			MeshRequest.data = Impl->MeshBlockRequests.GetData();
+			auto Result = MLMeshingRequestMesh(Impl->MeshTracker, &MeshRequest, &Impl->CurrentMeshRequest);
+			if (MLResult_Ok != Result)
 			{
-				//UE_LOG(LogMagicLeap, Log, TEXT("UMeshTrackerComponent::TickComponent found new group data."));
-				//LogMLDataArray(GroupData);
+				UE_LOG(LogMagicLeap, Error, TEXT("MLMeshingRequestMesh failed: %s."), UTF8_TO_TCHAR(MLGetResultString(Result)));
+				Impl->CurrentMeshInfoRequest = ML_INVALID_HANDLE;
+			}
+		}
 
-				Impl->bUpdateRequested = false;
-				MLDataArrayHandle* groupBegin = GroupData.streams[0].handle_array;
-				MLDataArrayHandle* groupEnd = GroupData.streams[0].handle_array + GroupData.streams[0].count;
+		// Request IDs and states of block meshes
+		if (Impl->CurrentMeshInfoRequest != ML_INVALID_HANDLE)
+		{
+			MLMeshingMeshInfo MeshInfo = {};
 
-				for (MLDataArrayHandle* groupIterator = groupBegin; groupIterator != groupEnd; ++groupIterator)
+			auto Result = MLMeshingGetMeshInfoResult(Impl->MeshTracker, Impl->CurrentMeshInfoRequest, &MeshInfo);
+			if (MLResult_Ok != Result)
+			{
+				// Just silently wait for pending result
+				if (MLResult_Pending != Result)
 				{
-					CurrentMeshHandleSet.Add(*groupIterator);
+					UE_LOG(LogMagicLeap, Error, 
+						TEXT("MLMeshingGetMeshInfoResult failed: %s."), UTF8_TO_TCHAR(MLGetResultString(Result)));
 				}
-
-				// Unlock group data array.
-				MLDataArrayUnlock(Impl->Data.meshes);
 			}
-
-			// Clear unnecessary meshes.
-			TMap<uint64, FMeshTrackerImpl::MeshCache> NewMeshHandleCacheMap;
-			for (const auto& MeshHandleCachePair : Impl->MeshHandleCacheMap)
+			else
 			{
-				if (!CurrentMeshHandleSet.Contains(MeshHandleCachePair.Key))
+				// Clear our stored block requests
+				Impl->MeshBlockRequests.Empty();
+
+				//UE_LOG(LogMagicLeap, Log, TEXT("Mesh tracker received %d mesh block infos"), MeshInfo.data_count);
+
+				for (uint32_t MeshInfoIndex = 0; MeshInfoIndex < MeshInfo.data_count; ++ MeshInfoIndex)
 				{
-					// We don't add any sections to the procedural mesh component for point clouds.
-					if (MeshType != EMeshType::PointCloud)
+					const auto &MeshInfoData = MeshInfo.data[MeshInfoIndex];
+
+					MLMeshingBlockRequest BlockRequest = {};
+
+					BlockRequest.id = MeshInfoData.id;
+					BlockRequest.level = MLToUnreal_MeshLOD(LevelOfDetail);
+
+					switch (MeshInfoData.state)
 					{
-						// MRMeshComponent may reference this data from multiple threads
-						const static TArray<FVector> EmptyVertices;
-						const static TArray<FVector2D> EmptyUVs;
-						const static TArray<FPackedNormal> EmptyTangents;
-						const static TArray<FColor> EmptyVertexColors;
-						const static TArray<uint32> EmptyTriangles;
-						if (MRMesh)
+						case MLMeshingMeshState_New:
+						case MLMeshingMeshState_Updated:
 						{
-							static_cast<IMRMesh*>(MRMesh)->SendBrickData(IMRMesh::FSendBrickDataArgs
+							// Store the block request so we can update it
+							Impl->MeshBlockRequests.Add(BlockRequest);
+							break;
+						}
+						case MLMeshingMeshState_Deleted:
+						{
+							// Delete the brick and its cache entry
+							if (Impl->MeshBrickCache.Contains(MeshInfoData.id))
 							{
-								nullptr,
-								MeshHandleCachePair.Key,
-								EmptyVertices,
-								EmptyUVs,
-								EmptyTangents,
-								EmptyVertexColors,
-								EmptyTriangles
+								const auto& BrickId = Impl->MeshBrickCache[MeshInfoData.id];
+
+								if (MeshType != EMeshType::PointCloud)
+								{
+									const static TArray<FVector> EmptyVertices;
+									const static TArray<FVector2D> EmptyUVs;
+									const static TArray<FPackedNormal> EmptyTangents;
+									const static TArray<FColor> EmptyVertexColors;
+									const static TArray<uint32> EmptyTriangles;
+									static_cast<IMRMesh*>(MRMesh)->SendBrickData(IMRMesh::FSendBrickDataArgs
+										{
+											nullptr,
+											BrickId,
+											EmptyVertices,
+											EmptyUVs,
+											EmptyTangents,
+											EmptyVertexColors,
+											EmptyTriangles
+										}
+									);
+								}
+
+								if (OnMeshTrackerUpdated.IsBound())
+								{
+									OnMeshTrackerUpdated.Broadcast(BrickId, 
+										TArray<FVector>(), TArray<int32>(), TArray<FVector>(), TArray<float>());
+								}
+
+								Impl->MeshBrickCache.Remove(MeshInfoData.id);
 							}
-							);
+							break;
 						}
+						default:
+							break;
 					}
-
-					OnMeshTrackerUpdated.Broadcast(MeshHandleCachePair.Value.Section, TArray<FVector>(), TArray<int32>(), TArray<FVector>(), TArray<float>());
-				}
-				else
-				{
-					NewMeshHandleCacheMap.Add(MeshHandleCachePair.Key, MeshHandleCachePair.Value);
 				}
+
+				// Free up the ML meshing resources
+				MLMeshingFreeResource(Impl->MeshTracker, &Impl->CurrentMeshInfoRequest);
+				Impl->CurrentMeshInfoRequest = ML_INVALID_HANDLE;
 			}
-			Impl->MeshHandleCacheMap = NewMeshHandleCacheMap;
+		}
+
+		// Request mesh data
+		if (Impl->CurrentMeshRequest != ML_INVALID_HANDLE)
+		{
+			MLMeshingMesh Mesh = {};
+
+			auto Result = MLMeshingGetMeshResult(Impl->MeshTracker, Impl->CurrentMeshRequest, &Mesh);
 
-			// Add new mesh handles.
-			for (auto CurrentMeshHandle : CurrentMeshHandleSet)
+			if (MLResult_Ok != Result)
 			{
-				if (!Impl->MeshHandleCacheMap.Contains(CurrentMeshHandle))
+				// Just silently wait for pending result
+				if (MLResult_Pending != Result)
 				{
-					FMeshTrackerImpl::MeshCache CurrentMeshCache;
-					FMemory::Memset(&CurrentMeshCache.Diff, 0, sizeof(MLDataArrayDiff));
-					CurrentMeshCache.Section = (Impl->SectionCounter)++;
-					Impl->MeshHandleCacheMap.Add(CurrentMeshHandle, CurrentMeshCache);
-					// We don't add any sections to the procedural mesh component for point clouds.
-					// If number of sections is more than number of materials for this component, set the first section's material as the new section's material.
+					UE_LOG(LogMagicLeap, Error, TEXT("MLMeshingGetMeshResult failed: %s."), UTF8_TO_TCHAR(MLGetResultString(Result)));
 				}
 			}
-		}
-
-		for (auto& MeshData : Impl->MeshHandleCacheMap)
-		{
-			// Lock mesh data array.
-			uint64 CurrentMeshHandle = MeshData.Key;
-			MLDataArray CurrentMeshData;
-			MLDataArrayLockResult meshDataArray = MLDataArrayTryLock(CurrentMeshHandle, &CurrentMeshData, &MeshData.Value.Diff);
-			if (meshDataArray == MLDataArrayLockResult_New)
+			else
 			{
-				// Collect current mesh data.
-				FMeshTrackerImpl::FMLCachedMeshData::SharedPtr CurrentMeshDataCache = Impl->AquireMeshDataCache();
-				CurrentMeshDataCache->BrickId = CurrentMeshHandle;
+				//UE_LOG(LogMagicLeap, Log, TEXT("Mesh tracker received %d blocks for mesh request %d"), 
+					//Mesh.data_count, static_cast<int>(Impl->CurrentMeshRequest));
 
-				// Work with the locked mesh MLDataArray.
+				for (uint32_t MeshIndex = 0; MeshIndex < Mesh.data_count; ++ MeshIndex)
 				{
-					//UE_LOG(LogMagicLeap, Log, TEXT("UMeshTrackerComponent::TickComponent found new mesh data."));
-					//LogMLDataArray(CurrentMeshData);
-
-					Impl->bUpdateRequested = false;
+					const auto &MeshData = Mesh.data[MeshIndex];
 
-					// TODO: Copying the position and normal arrays could be a mem copy if we're sure about the memory layout of things.
-
-					// Read position stream
+					// Create a brick index for any new mesh block
+					if (!Impl->MeshBrickCache.Contains(MeshData.id))
 					{
-						const uint32 PositionCount = CurrentMeshData.streams[Impl->Data.position_stream_index].count;
-						CurrentMeshDataCache->Vertices.Reserve(PositionCount);
-						for (uint32 i = 0; i < PositionCount; ++i)
-						{
-							const auto &position = CurrentMeshData.streams[Impl->Data.position_stream_index].xyz_array[i];
-							CurrentMeshDataCache->Vertices.Add(MagicLeap::ToFVector(position, WorldToMetersScale));
-						}
+						Impl->MeshBrickCache.Add(MeshData.id, Impl->MeshBrickIndex ++);
 					}
 
-					// Read triangle (indices) stream
-					if (Impl->Data.triangle_index_stream_index == MLDataArray_InvalidStreamIndex || CurrentMeshData.streams[Impl->Data.triangle_index_stream_index].uint_array == nullptr)
+					auto BrickId = Impl->MeshBrickCache[MeshData.id];
+
+					// Acquire mesh data cache and mark its brick ID
+					FMeshTrackerImpl::FMLCachedMeshData::SharedPtr CurrentMeshDataCache = Impl->AquireMeshDataCache();
+					CurrentMeshDataCache->BrickId = BrickId;
+
+					// Pull vertices
+					CurrentMeshDataCache->Vertices.Reserve(MeshData.vertex_count);
+					for (uint32_t v = 0; v < MeshData.vertex_count; ++ v)
 					{
-						const int32 TriangleCount = CurrentMeshData.streams[Impl->Data.triangle_index_stream_index].count;
-						CurrentMeshDataCache->Triangles.AddZeroed(TriangleCount);
+						CurrentMeshDataCache->Vertices.Add(MagicLeap::ToFVector(MeshData.vertex[v], WorldToMetersScale));
 					}
-					else
+
+					// Pull indices
+					CurrentMeshDataCache->Triangles.Reserve(MeshData.index_count);
+					for (uint16_t i = 0; i < MeshData.index_count; ++ i)
 					{
-						CurrentMeshDataCache->Triangles.Append(CurrentMeshData.streams[Impl->Data.triangle_index_stream_index].uint_array, CurrentMeshData.streams[Impl->Data.triangle_index_stream_index].count);
+						CurrentMeshDataCache->Triangles.Add(static_cast<uint32>(MeshData.index[i]));
 					}
 
-					// Read normal stream
+					// Pull normals
+					CurrentMeshDataCache->Normals.Reserve(MeshData.vertex_count);
+					if (nullptr != MeshData.normal)
 					{
-						uint32 NormalsCount = 0;
-						if (Impl->Data.normal_stream_index != MLDataArray_InvalidStreamIndex && CurrentMeshData.streams[Impl->Data.normal_stream_index].type != MLDataArrayType_None)
+						for (uint32_t n = 0; n < MeshData.vertex_count; ++ n)
 						{
-							NormalsCount = CurrentMeshData.streams[Impl->Data.normal_stream_index].count;
-							CurrentMeshDataCache->Normals.Reserve(NormalsCount);
-							for (uint32 i = 0; i < NormalsCount; ++i)
-							{
-								const auto &normal = CurrentMeshData.streams[Impl->Data.normal_stream_index].xyz_array[i];
-								CurrentMeshDataCache->Normals.Add(MagicLeap::ToFVector(normal, 1.0f));
-							}
+							CurrentMeshDataCache->Normals.Add(MagicLeap::
+								ToFVector(MeshData.normal[n], WorldToMetersScale));
 						}
-						else
+					}
+					// If no normals were provided we need to pack fake ones for Vulkan
+					else
+					{
+						for (uint32_t n = 0; n < MeshData.vertex_count; ++ n)
 						{
-							// Vulkan requires that tangent data exist, so we provide fake normals and tangents even if it does not
-
-							NormalsCount = CurrentMeshDataCache->Vertices.Num();
-							CurrentMeshDataCache->Normals.Reserve(NormalsCount);
-							for (uint32 i = 0; i < NormalsCount; ++i)
-							{
-								FVector Normal = CurrentMeshDataCache->Vertices[i];
-								Normal.Normalize();
-								CurrentMeshDataCache->Normals.Add(Normal);
-							}
+							FVector FakeNormal = CurrentMeshDataCache->Vertices[n];
+							FakeNormal.Normalize();
+							CurrentMeshDataCache->Normals.Add(FakeNormal);
 						}
+					}
 
-						// Also write normals into tangents
-						CurrentMeshDataCache->Tangents.Reserve(NormalsCount * 2);
-						for (uint32 i = 0; i < NormalsCount; ++i)
-						{
-							const FVector Normal = CurrentMeshDataCache->Normals[i];
-							const FVector NonNormal = Normal.X < Normal.Z ? FVector(0, 0, 1) : FVector(0, 1, 0);
-							const FVector TangentX = FVector::CrossProduct(Normal, NonNormal);
-							CurrentMeshDataCache->Tangents.Add(TangentX);
-							CurrentMeshDataCache->Tangents.Add(Normal);
-						}
+					// Calculate and pack tangents
+					CurrentMeshDataCache->Tangents.Reserve(MeshData.vertex_count * 2);
+					for (uint32_t t = 0; t < MeshData.vertex_count; ++ t)
+					{
+						const FVector& Norm = CurrentMeshDataCache->Normals[t];
 
+						// Calculate tangent
+						auto Perp = Norm.X < Norm.Z ? 
+							FVector(1.0f, 0.0f, 0.0f) : FVector(0.0f, 1.0f, 0.0f);
+						auto Tang = FVector::CrossProduct(Norm, Perp);
+
+						CurrentMeshDataCache->Tangents.Add(Tang);
+						CurrentMeshDataCache->Tangents.Add(Norm);
 					}
 
-					// Read confidence stream
-					if (Impl->Data.confidence_stream_index != MLDataArray_InvalidStreamIndex && CurrentMeshData.streams[Impl->Data.confidence_stream_index].type != MLDataArrayType_None)
+					// Pull confidence
+					if (nullptr != MeshData.confidence)
 					{
-						const uint32 ConfidenceCount = CurrentMeshData.streams[Impl->Data.confidence_stream_index].count;
-						CurrentMeshDataCache->Confidence.Reserve(ConfidenceCount);
-						for (uint32 i = 0; i < ConfidenceCount; ++i)
-						{
-							const float Confidence = CurrentMeshData.streams[Impl->Data.confidence_stream_index].float_array[i];
-							CurrentMeshDataCache->Confidence.Add(Confidence);
-						}
+						CurrentMeshDataCache->Confidence.Append(MeshData.confidence, MeshData.vertex_count);
 					}
 
-					// Write VertexColor
+					// Apply chosen vertex color mode
+					switch (VertexColorMode)
 					{
-						switch (VertexColorMode)
-						{
 						case EMLMeshVertexColorMode::Confidence:
 						{
-							const uint32 ConfidenceCount = CurrentMeshDataCache->Confidence.Num();
-							const uint32 PositionCount = CurrentMeshDataCache->Vertices.Num();
-							// Length of VertexColor array needs to be same as length of Position array.
-							CurrentMeshDataCache->VertexColors.Reserve(PositionCount);
-							for (uint32 i = 0; i < ConfidenceCount; ++i)
+							if (nullptr != MeshData.confidence)
 							{
-								const float Confidence = CurrentMeshDataCache->Confidence[i];
-								const FLinearColor VertexColor = FMath::Lerp(VertexColorFromConfidenceZero, VertexColorFromConfidenceOne, Confidence);
-								CurrentMeshDataCache->VertexColors.Add(VertexColor.ToFColor(false));
+								CurrentMeshDataCache->VertexColors.Reserve(MeshData.vertex_count);
+								for (uint32 v = 0; v < MeshData.vertex_count; ++ v)
+								{
+									const FLinearColor VertexColor = FMath::Lerp(VertexColorFromConfidenceZero, 
+										VertexColorFromConfidenceOne, CurrentMeshDataCache->Confidence[v]);
+									CurrentMeshDataCache->VertexColors.Add(VertexColor.ToFColor(false));
+								}
 							}
-							for (uint32 i = ConfidenceCount; i < PositionCount; ++i)
+							else
 							{
-								CurrentMeshDataCache->VertexColors.Add(VertexColorFromConfidenceZero.ToFColor(false));
+								UE_LOG(LogMagicLeap, Warning, TEXT("MeshTracker vertex color mode is Confidence "
+									"but no confidence values available. Using white for all blocks."));
 							}
 							break;
 						}
 						case EMLMeshVertexColorMode::Block:
 						{
-							const uint32 VertexCount = CurrentMeshDataCache->Vertices.Num();
-							const uint32 ColorCount = BlockVertexColors.Num();
-							const FColor VertexColor = ColorCount > 0 ? BlockVertexColors[MeshData.Value.Section % ColorCount] : FColor::White;
-							if (BlockVertexColors.Num() == 0)
+							if (BlockVertexColors.Num() > 0)
 							{
-								UE_LOG(LogMagicLeap, Warning, TEXT("MeshTracker is in EMLMeshVertexColorMode::Block, but has no BlockVertexColors set.  Using White for all blocks."));
+								const FColor& VertexColor = BlockVertexColors[BrickId % BlockVertexColors.Num()];
+
+								CurrentMeshDataCache->VertexColors.Reserve(MeshData.vertex_count);
+								for (uint32 v = 0; v < MeshData.vertex_count; ++ v)
+								{
+									CurrentMeshDataCache->VertexColors.Add(VertexColor);
+								}
 							}
-							for (uint32 i = 0; i < VertexCount; ++i)
+							else
 							{
-								CurrentMeshDataCache->VertexColors.Add(VertexColor);
+								UE_LOG(LogMagicLeap, Warning, TEXT("MeshTracker vertex color mode is Block but "
+									"no BlockVertexColors set. Using white for all blocks."));
 							}
 							break;
 						}
 						case EMLMeshVertexColorMode::None:
 						{
-							// Vulkan requires that we fill everything in.
-							const uint32 VertexCount = CurrentMeshDataCache->Vertices.Num();
-							for (uint32 i = 0; i < VertexCount; ++i)
-							{
-								CurrentMeshDataCache->VertexColors.Add(FColor::White);
-							}
 							break;
 						}
 						default:
 							check(false);
-						}
-					}
-
-					// Unlock mesh data array.
-					MLDataArrayUnlock(CurrentMeshHandle);
-				}
-
-				// Write UVs
-				{
-					const uint32 VertexCount = CurrentMeshDataCache->Vertices.Num();
-					for (uint32 i = 0; i < VertexCount; ++i)
-					{
-						const float FakeCoord = (float)i / (float)VertexCount;
-						CurrentMeshDataCache->UV0.Add(FVector2D(FakeCoord, FakeCoord));
+							break;
 					}
-				}
-
-				// We don't add any sections to the procedural mesh component for point clouds.
-				if (MeshType != EMeshType::PointCloud)
-				{
-					// Create or Update. Create on existing mesh section will overwrite that section.
 
-					if (MRMesh)
+					// To work in all rendering paths we always set a vertex color
+					if (CurrentMeshDataCache->VertexColors.Num() == 0)
 					{
-						static_cast<IMRMesh*>(MRMesh)->SendBrickData(IMRMesh::FSendBrickDataArgs
+						for (uint32 v = 0; v < MeshData.vertex_count; ++ v)
 						{
-							TSharedPtr<IMRMesh::FBrickDataReceipt, ESPMode::ThreadSafe>(new FMeshTrackerImpl::FMeshTrackerComponentBrickDataReceipt(CurrentMeshDataCache)),
-							CurrentMeshDataCache->BrickId,
-							CurrentMeshDataCache->Vertices,
-							CurrentMeshDataCache->UV0,
-							CurrentMeshDataCache->Tangents,
-							CurrentMeshDataCache->VertexColors,
-							CurrentMeshDataCache->Triangles
+							CurrentMeshDataCache->VertexColors.Add(FColor::White);
 						}
-						);
 					}
-				}
-				if (OnMeshTrackerUpdated.IsBound())
-				{
-					// Hack because blueprints don't support uint32.
-					TArray<int32> Triangles(reinterpret_cast<const int32*>(CurrentMeshDataCache->Triangles.GetData()), CurrentMeshDataCache->Triangles.Num());
-					OnMeshTrackerUpdated.Broadcast(MeshData.Value.Section, CurrentMeshDataCache->Vertices, Triangles, CurrentMeshDataCache->Normals, CurrentMeshDataCache->Confidence);
-				}
-			}
-		}
-	}
-#endif //WITH_MLSDK
-}
-
-void UMeshTrackerComponent::TickWithFakeData()
-{
-#if WITH_MLSDK
-	// fake only works at default worldscale
-	float WorldToMetersScale = 100;
-
-	FTransform Pose = FTransform::Identity;
-	if (MRMesh)
-	{
-		MRMesh->SendRelativeTransform(Pose);
-	}
 
-	// Only draw meshes during the scan phase in order to maintain a steady framerate.
-	if (ScanWorld || Impl->bUpdateRequested)
-	{
-		// Collect current mesh handles.
-		TSet<uint64> CurrentMeshHandleSet;
-
-		// Figure out new fake data mode, so we can cycle through a few different setups.
-		static int period = 60;
-		static int count = 0;
-		count = (count + 1) % (period * 2);
-		bool bTest = count < period;
-		static bool bOldTest = false;
-
-		static int modes = 3;
-		bool bModeChanged = false;
-		static int oldMode = -1;
-		int newMode = 0;
-		if (bOldTest != bTest)
-		{
-			bOldTest = bTest;
-			bModeChanged = true;
-
-			newMode = (oldMode + 1) % modes;
-			oldMode = newMode;
-		}
-		
-		if (bModeChanged)
-		{
-			UE_LOG(LogMagicLeap, Display, TEXT("TickWithFakeData() mode changed to %i"), newMode);
-			Impl->bUpdateRequested = false;
-
-			if (newMode != 0)
-			{
-				CurrentMeshHandleSet.Add(newMode);
-			}
+					// Write UVs
+					CurrentMeshDataCache->UV0.Reserve(MeshData.vertex_count);
+					for (uint32 v = 0; v < MeshData.vertex_count; ++ v)
+					{
+						const float FakeCoord = static_cast<float>(v) / static_cast<float>(MeshData.vertex_count);
+						CurrentMeshDataCache->UV0.Add(FVector2D(FakeCoord, FakeCoord));
+					}
 
-			// Clear unnecessary meshes.
-			TMap<uint64, FMeshTrackerImpl::MeshCache> NewMeshHandleCacheMap;
-			for (const auto& MeshHandleCachePair : Impl->MeshHandleCacheMap)
-			{
-				if (!CurrentMeshHandleSet.Contains(MeshHandleCachePair.Key))
-				{
-					// We don't add any sections to the procedural mesh component for point clouds.
-					//if (MeshType != EMeshType::PointCloud)
+					// Create/update brick
+					if (MeshType != EMeshType::PointCloud)
 					{
-						static TArray<FVector> EmptyVertices;
-						const static TArray<FVector2D> EmptyUVs;
-						const static TArray<FPackedNormal> EmptyTangents;
-						static TArray<FColor> EmptyVertexColors;
-						static TArray<uint32> EmptyTriangles;
-						if (MRMesh)
-						{
-							UE_LOG(LogMagicLeap, Display, TEXT("TickWithFakeData() sending empty brick to remove %llu"), MeshHandleCachePair.Key);
-							static_cast<IMRMesh*>(MRMesh)->SendBrickData(IMRMesh::FSendBrickDataArgs
+						static_cast<IMRMesh*>(MRMesh)->SendBrickData(IMRMesh::FSendBrickDataArgs
 							{
-								nullptr,
-								MeshHandleCachePair.Key,
-								EmptyVertices,
-								EmptyUVs,
-								EmptyTangents,
-								EmptyVertexColors,
-								EmptyTriangles
+								TSharedPtr<IMRMesh::FBrickDataReceipt, ESPMode::ThreadSafe>
+									(new FMeshTrackerImpl::FMeshTrackerComponentBrickDataReceipt(CurrentMeshDataCache)),
+								CurrentMeshDataCache->BrickId,
+								CurrentMeshDataCache->Vertices,
+								CurrentMeshDataCache->UV0,
+								CurrentMeshDataCache->Tangents,
+								CurrentMeshDataCache->VertexColors,
+								CurrentMeshDataCache->Triangles
 							}
-							);
-						}
+						);
 					}
-					OnMeshTrackerUpdated.Broadcast(MeshHandleCachePair.Value.Section, TArray<FVector>(), TArray<int32>(), TArray<FVector>(), TArray<float>());
-				}
-				else
-				{
-					NewMeshHandleCacheMap.Add(MeshHandleCachePair.Key, MeshHandleCachePair.Value);
-				}
-			}
-			Impl->MeshHandleCacheMap = NewMeshHandleCacheMap;
 
-			// Add new mesh handles.
-			for (auto CurrentMeshHandle : CurrentMeshHandleSet)
-			{
-				if (!Impl->MeshHandleCacheMap.Contains(CurrentMeshHandle))
-				{
-					FMeshTrackerImpl::MeshCache CurrentMeshCache;
-					memset(&CurrentMeshCache.Diff, 0, sizeof(MLDataArrayDiff));
-					CurrentMeshCache.Section = (Impl->SectionCounter)++;
-					Impl->MeshHandleCacheMap.Add(CurrentMeshHandle, CurrentMeshCache);
-				}
-			}
-		}
-
-		for (auto& MeshData : Impl->MeshHandleCacheMap)
-		{
-			// Lock mesh data array.
-			uint64 CurrentMeshHandle = MeshData.Key;
-			if (bModeChanged && CurrentMeshHandle == newMode)
-			{
-				Impl->bUpdateRequested = false;
-				// Collect current mesh data.
-				FMeshTrackerImpl::FMLCachedMeshData::SharedPtr CurrentMeshDataCache = Impl->AquireMeshDataCache();
-				CurrentMeshDataCache->BrickId = CurrentMeshHandle;
-
-				FVector Origin = FVector((CurrentMeshHandle + 1) * 50, 0, 0);
-				static FVector Extents(20.0f, 20.0f, 50.0f);
-
-				const int32 VertCount = 8;
-				CurrentMeshDataCache->Vertices.Reserve(VertCount);
-
-				CurrentMeshDataCache->Vertices.Add(FVector(Origin.X + Extents.X, Origin.Y - Extents.Y, Origin.Z + Extents.Z));  // 0
-				CurrentMeshDataCache->Vertices.Add(FVector(Origin.X + Extents.X, Origin.Y + Extents.Y, Origin.Z + Extents.Z));  // 1
-				CurrentMeshDataCache->Vertices.Add(FVector(Origin.X + Extents.X, Origin.Y + Extents.Y, Origin.Z - Extents.Z));  // 2
-				CurrentMeshDataCache->Vertices.Add(FVector(Origin.X + Extents.X, Origin.Y - Extents.Y, Origin.Z - Extents.Z));  // 3
-				CurrentMeshDataCache->Vertices.Add(FVector(Origin.X - Extents.X, Origin.Y - Extents.Y, Origin.Z + Extents.Z));  // 4
-				CurrentMeshDataCache->Vertices.Add(FVector(Origin.X - Extents.X, Origin.Y + Extents.Y, Origin.Z + Extents.Z));  // 5
-				CurrentMeshDataCache->Vertices.Add(FVector(Origin.X - Extents.X, Origin.Y + Extents.Y, Origin.Z - Extents.Z));  // 6
-				CurrentMeshDataCache->Vertices.Add(FVector(Origin.X - Extents.X, Origin.Y - Extents.Y, Origin.Z - Extents.Z));  // 7
-
-				CurrentMeshDataCache->UV0.Reserve(VertCount);
-				CurrentMeshDataCache->VertexColors.Reserve(VertCount);
-				for (int i = 0; i < VertCount; ++i)
-				{
-					int imod = i % 10;
-					CurrentMeshDataCache->VertexColors.Emplace(FColor::MakeRandomColor());
-					CurrentMeshDataCache->UV0.Add(FVector2D(imod*0.1f, imod*0.1f));
-				}
-				
-				// Vulkan requires that tangent data exist
-
-				CurrentMeshDataCache->Normals.Reserve(VertCount);
-				for (uint32 i = 0; i < VertCount; ++i)
-				{
-					CurrentMeshDataCache->Normals.Add(FVector(0,0,1));
-				}
-
-				// Also write normals into tangents
-				CurrentMeshDataCache->Tangents.Reserve(VertCount * 2);
-				for (uint32 i = 0; i < VertCount; ++i)
-				{
-					const FVector Normal = CurrentMeshDataCache->Normals[i];
-					const FVector NonNormal = Normal.X < Normal.Z ? FVector(0, 0, 1) : FVector(0, 1, 0);
-					const FVector TangentX = FVector::CrossProduct(Normal, NonNormal);
-					CurrentMeshDataCache->Tangents.Add(TangentX);
-					CurrentMeshDataCache->Tangents.Add(Normal);
-				}
-				
-				CurrentMeshDataCache->Triangles.Reserve(6 * 2 * 3);
-				// Read triangle (indices) stream
-				// This makes half the triangles of a box, one per side.
-				//CurrentMeshDataCache->Triangles.Add(0);
-				//CurrentMeshDataCache->Triangles.Add(1);
-				//CurrentMeshDataCache->Triangles.Add(2);
-
-				CurrentMeshDataCache->Triangles.Add(0);
-				CurrentMeshDataCache->Triangles.Add(2);
-				CurrentMeshDataCache->Triangles.Add(3);
-
-				//CurrentMeshDataCache->Triangles.Add(0);
-				//CurrentMeshDataCache->Triangles.Add(4);
-				//CurrentMeshDataCache->Triangles.Add(1);
-
-				CurrentMeshDataCache->Triangles.Add(1);
-				CurrentMeshDataCache->Triangles.Add(4);
-				CurrentMeshDataCache->Triangles.Add(5);
-
-				//CurrentMeshDataCache->Triangles.Add(7);
-				//CurrentMeshDataCache->Triangles.Add(5);
-				//CurrentMeshDataCache->Triangles.Add(4);
-
-				CurrentMeshDataCache->Triangles.Add(6);
-				CurrentMeshDataCache->Triangles.Add(5);
-				CurrentMeshDataCache->Triangles.Add(7);
-
-				//CurrentMeshDataCache->Triangles.Add(7);
-				//CurrentMeshDataCache->Triangles.Add(3);
-				//CurrentMeshDataCache->Triangles.Add(2);
-
-				CurrentMeshDataCache->Triangles.Add(7);
-				CurrentMeshDataCache->Triangles.Add(2);
-				CurrentMeshDataCache->Triangles.Add(6);
-
-				//CurrentMeshDataCache->Triangles.Add(7);
-				//CurrentMeshDataCache->Triangles.Add(4);
-				//CurrentMeshDataCache->Triangles.Add(0);
-
-				CurrentMeshDataCache->Triangles.Add(7);
-				CurrentMeshDataCache->Triangles.Add(0);
-				CurrentMeshDataCache->Triangles.Add(3);
-
-				//CurrentMeshDataCache->Triangles.Add(1);
-				//CurrentMeshDataCache->Triangles.Add(5);
-				//CurrentMeshDataCache->Triangles.Add(6);
-
-				CurrentMeshDataCache->Triangles.Add(2);
-				CurrentMeshDataCache->Triangles.Add(1);
-				CurrentMeshDataCache->Triangles.Add(6);
-
-
-				// We don't add any sections to the mesh component for point clouds.
-				if (MeshType != EMeshType::PointCloud)
-				{
-					if (MRMesh)
+					// Broadcast that a mesh was updated
+					if (OnMeshTrackerUpdated.IsBound())
 					{
-						UE_LOG(LogMagicLeap, Error, TEXT("TickWithFakeData() Sending fake brick %llu."), CurrentMeshDataCache->BrickId);
-						static_cast<IMRMesh*>(MRMesh)->SendBrickData(IMRMesh::FSendBrickDataArgs
-						{
-							TSharedPtr<IMRMesh::FBrickDataReceipt, ESPMode::ThreadSafe>(new FMeshTrackerImpl::FMeshTrackerComponentBrickDataReceipt(CurrentMeshDataCache)),
-							CurrentMeshDataCache->BrickId,
-							CurrentMeshDataCache->Vertices,
-							CurrentMeshDataCache->UV0,
-							CurrentMeshDataCache->Tangents,
-							CurrentMeshDataCache->VertexColors,
-							CurrentMeshDataCache->Triangles
-						}
-						);
+						// Hack because blueprints don't support uint32.
+						TArray<int32> Triangles(reinterpret_cast<const int32*>(CurrentMeshDataCache->
+							Triangles.GetData()), CurrentMeshDataCache->Triangles.Num());
+						OnMeshTrackerUpdated.Broadcast(CurrentMeshDataCache->BrickId, CurrentMeshDataCache->Vertices, 
+							Triangles, CurrentMeshDataCache->Normals, CurrentMeshDataCache->Confidence);
 					}
 				}
-				if (OnMeshTrackerUpdated.IsBound())
-				{
-					// Hack because blueprints don't support uint32.
-					TArray<int32> Triangles(reinterpret_cast<const int32*>(CurrentMeshDataCache->Triangles.GetData()), CurrentMeshDataCache->Triangles.Num());
-					OnMeshTrackerUpdated.Broadcast(MeshData.Value.Section, CurrentMeshDataCache->Vertices, Triangles, CurrentMeshDataCache->Normals, CurrentMeshDataCache->Confidence);
-				}
+
+				// All meshes pulled and/or updated; free the ML resource
+				MLMeshingFreeResource(Impl->MeshTracker, &Impl->CurrentMeshRequest);
+				Impl->CurrentMeshRequest = ML_INVALID_HANDLE;
 			}
 		}
 	}
diff --git a/Engine/Plugins/Lumin/MagicLeap/Source/MagicLeap/Private/PlanesComponent.cpp b/Engine/Plugins/Lumin/MagicLeap/Source/MagicLeap/Private/PlanesComponent.cpp
--- a/Engine/Plugins/Lumin/MagicLeap/Source/MagicLeap/Private/PlanesComponent.cpp
+++ b/Engine/Plugins/Lumin/MagicLeap/Source/MagicLeap/Private/PlanesComponent.cpp
@@ -37,9 +37,9 @@ public:
 		if (!MLHandleIsValid(Tracker))
 		{
 			UE_LOG(LogMagicLeap, Display, TEXT("Creating Planes Tracker"));
-			Tracker = MLPlanesCreate();
+			MLResult CreateResult = MLPlanesCreate(&Tracker);
 
-			if (!MLHandleIsValid(Tracker))
+			if (CreateResult != MLResult_Ok || !MLHandleIsValid(Tracker))
 			{
 				UE_LOG(LogMagicLeap, Error, TEXT("Could not create planes tracker."));
 				return false;
@@ -54,8 +54,8 @@ public:
 #if WITH_MLSDK
 		if (MLHandleIsValid(Tracker))
 		{
-			bool bResult = MLPlanesDestroy(Tracker);
-			if (!bResult)
+			MLResult DestroyResult = MLPlanesDestroy(Tracker);
+			if (DestroyResult != MLResult_Ok)
 			{
 				UE_LOG(LogMagicLeap, Error, TEXT("Error destroying planes tracker."));
 			}
@@ -199,7 +199,7 @@ void UPlanesComponent::TickComponent(float DeltaTime, enum ELevelTick TickType,
 #if WITH_MLSDK
 	Super::TickComponent(DeltaTime, TickType, ThisTickFunction);
 
-	if (!Impl->Create())
+	if (!(IMagicLeapPlugin::Get().IsMagicLeapHMDValid() && Impl->Create()))
 	{
 		return;
 	}
@@ -212,16 +212,19 @@ void UPlanesComponent::TickComponent(float DeltaTime, enum ELevelTick TickType,
 		TArray<MLPlane> resultMLPlanes;
 		resultMLPlanes.AddDefaulted(pair.Value.MaxResults);
 
-		MLPlanesQueryResult result = MLPlanesQueryGetResults(Impl->Tracker, pair.Key, resultMLPlanes.GetData(), &outNumResults);
+		MLResult result = MLPlanesQueryGetResults(Impl->Tracker, pair.Key, resultMLPlanes.GetData(), &outNumResults);
 		switch (result)
 		{
-		case MLPlanesQueryResult_Failure:
+		case MLResult_Pending:
+			// Intentionally skip. We'll continue to check until it has completed.
+			break;
+		case MLResult_UnspecifiedFailure:
 		{
 			pair.Value.ResultDelegate.ExecuteIfBound(false, TArray<FPlaneResult>(), pair.Value.UserData);
 			CompletedRequests.Add(pair.Key);
 			break;
 		}
-		case MLPlanesQueryResult_Success:
+		case MLResult_Ok:
 		{
 			const FAppFramework& AppFramework = static_cast<FMagicLeapHMD*>(GEngine->XRSystem->GetHMDDevice())->GetAppFrameworkConst();
 			float WorldToMetersScale = AppFramework.GetWorldToMetersScale();
@@ -267,6 +270,8 @@ void UPlanesComponent::TickComponent(float DeltaTime, enum ELevelTick TickType,
 			CompletedRequests.Add(pair.Key);
 			break;
 		}
+		default:
+			UE_LOG(LogMagicLeap, Warning, TEXT("Unexpected return code from MLPlanesQueryGetResults: %d"), result);		
 		}
 	}
 
@@ -285,6 +290,11 @@ void UPlanesComponent::TickComponent(float DeltaTime, enum ELevelTick TickType,
 bool UPlanesComponent::RequestPlanes(int32 UserData, const FPlaneResultDelegate& ResultDelegate)
 {
 #if WITH_MLSDK
+	if (!IMagicLeapPlugin::Get().IsMagicLeapHMDValid())
+	{
+		return false;
+	}
+
 	const FAppFramework& AppFramework = static_cast<FMagicLeapHMD*>(GEngine->XRSystem->GetHMDDevice())->GetAppFrameworkConst();
 	float WorldToMetersScale = AppFramework.GetWorldToMetersScale();
 	check(WorldToMetersScale != 0);
@@ -312,8 +322,9 @@ bool UPlanesComponent::RequestPlanes(int32 UserData, const FPlaneResultDelegate&
 		query.bounds_extents.z = FMath::Abs<float>(query.bounds_extents.z);
 	}
 
-	MLHandle handle = MLPlanesQueryBegin(Impl->Tracker, &query);
-	if (!MLHandleIsValid(handle))
+	MLHandle handle;
+	MLResult QueryResult = MLPlanesQueryBegin(Impl->Tracker, &query, &handle);
+	if (QueryResult != MLResult_Ok || !MLHandleIsValid(handle))
 	{
 		UE_LOG(LogMagicLeap, Error, TEXT("Could not request planes."));
 		return false;
diff --git a/Engine/Plugins/Lumin/MagicLeap/Source/MagicLeap/Private/RaycastComponent.cpp b/Engine/Plugins/Lumin/MagicLeap/Source/MagicLeap/Private/RaycastComponent.cpp
--- a/Engine/Plugins/Lumin/MagicLeap/Source/MagicLeap/Private/RaycastComponent.cpp
+++ b/Engine/Plugins/Lumin/MagicLeap/Source/MagicLeap/Private/RaycastComponent.cpp
@@ -37,10 +37,10 @@ public:
 #if WITH_MLSDK
 		if (!MLHandleIsValid(Tracker))
 		{
-			Tracker = MLRaycastCreate();
-			if (!MLHandleIsValid(Tracker))
+			MLResult Result = MLRaycastCreate(&Tracker);
+			if (Result != MLResult_Ok)
 			{
-				UE_LOG(LogMagicLeap, Error, TEXT("Error creating raycast tracker."));
+				UE_LOG(LogMagicLeap, Error, TEXT("MLRaycastCreate failed with error %d."), Result);
 				return false;
 			}
 		}
@@ -53,11 +53,8 @@ public:
 #if WITH_MLSDK
 		if (MLHandleIsValid(Tracker))
 		{
-			bool bResult = MLRaycastDestroy(Tracker);
-			if (!bResult)
-			{
-				UE_LOG(LogMagicLeap, Error, TEXT("Error destroying raycast tracker."));
-			}
+			MLResult Result = MLRaycastDestroy(Tracker);
+			UE_CLOG(Result != MLResult_Ok, LogMagicLeap, Error, TEXT("MLRaycastDestroy failed with error %d."), Result);
 			Tracker = ML_INVALID_HANDLE;
 		}
 #endif //WITH_MLSDK
@@ -76,7 +73,7 @@ ERaycastResultState MLToUnrealRaycastResultState(MLRaycastResultState state)
 	case MLRaycastResultState_HitUnobserved:
 		return ERaycastResultState::HitUnobserved;
 	case MLRaycastResultState_NoCollision:
-		return ERaycastResultState::NoCollission;
+		return ERaycastResultState::NoCollision;
 	}
 	return ERaycastResultState::RequestFailed;
 }
@@ -107,9 +104,9 @@ URaycastComponent::~URaycastComponent()
 void URaycastComponent::TickComponent(float DeltaTime, enum ELevelTick TickType, FActorComponentTickFunction *ThisTickFunction)
 {
 	Super::TickComponent(DeltaTime, TickType, ThisTickFunction);
-#if WITH_MLSDK
 
-	if (!MLHandleIsValid(Impl->Tracker))
+#if WITH_MLSDK
+	if (!(IMagicLeapPlugin::Get().IsMagicLeapHMDValid() && MLHandleIsValid(Impl->Tracker)))
 	{
 		return;
 	}
@@ -117,7 +114,8 @@ void URaycastComponent::TickComponent(float DeltaTime, enum ELevelTick TickType,
 	for (auto& pair : PendingRequests)
 	{
 		MLRaycastResult result;
-		if (MLRaycastGetResult(Impl->Tracker, pair.Key, &result))
+		MLResult APICallResult = MLRaycastGetResult(Impl->Tracker, pair.Key, &result);
+		if (APICallResult == MLResult_Ok)
 		{
 			const FAppFramework& AppFramework = static_cast<FMagicLeapHMD*>(GEngine->XRSystem->GetHMDDevice())->GetAppFrameworkConst();
 			float WorldToMetersScale = AppFramework.IsInitialized() ? AppFramework.GetWorldToMetersScale() : 100.0f;
@@ -142,6 +140,10 @@ void URaycastComponent::TickComponent(float DeltaTime, enum ELevelTick TickType,
 			pair.Value.ResultDelegate.ExecuteIfBound(hitResult);
 			CompletedRequests.Add(pair.Key);
 		}
+		else if (APICallResult != MLResult_Pending)
+		{
+			UE_LOG(LogMagicLeap, Error, TEXT("MLRaycastGetResult failed with result %d."), APICallResult);
+		}
 	}
 
 	// TODO: Implement better strategy to optimize memory allocation.
@@ -159,7 +161,7 @@ void URaycastComponent::TickComponent(float DeltaTime, enum ELevelTick TickType,
 bool URaycastComponent::RequestRaycast(const FRaycastQueryParams& RequestParams, const FRaycastResultDelegate& ResultDelegate)
 {
 #if WITH_MLSDK
-	if (!Impl->Create())
+	if (!(IMagicLeapPlugin::Get().IsMagicLeapHMDValid() && Impl->Create()))
 	{
 		return false;
 	}
@@ -178,14 +180,15 @@ bool URaycastComponent::RequestRaycast(const FRaycastQueryParams& RequestParams,
 	query.collide_with_unobserved = RequestParams.CollideWithUnobserved;
 	query.horizontal_fov_degrees = RequestParams.HorizontalFovDegrees;
 
-	MLHandle handle = MLRaycastRequest(Impl->Tracker, &query);
-	if (!MLHandleIsValid(handle))
+	MLHandle Handle = ML_INVALID_HANDLE;
+	MLResult Result = MLRaycastRequest(Impl->Tracker, &query, &Handle);
+	if (Result != MLResult_Ok)
 	{
-		UE_LOG(LogMagicLeap, Error, TEXT("Could not request raycast."));
+		UE_LOG(LogMagicLeap, Error, TEXT("MLRaycastRequest failed with error %d."), Result);
 		return false;
 	}
 
-	FRaycastRequestMetaData& requestMetaData = PendingRequests.Add(handle);
+	FRaycastRequestMetaData& requestMetaData = PendingRequests.Add(Handle);
 	requestMetaData.UserData = RequestParams.UserData;
 	requestMetaData.ResultDelegate = ResultDelegate;
 #endif //WITH_MLSDK
diff --git a/Engine/Plugins/Lumin/MagicLeap/Source/MagicLeap/Public/IMagicLeapPlugin.h b/Engine/Plugins/Lumin/MagicLeap/Source/MagicLeap/Public/IMagicLeapPlugin.h
--- a/Engine/Plugins/Lumin/MagicLeap/Source/MagicLeap/Public/IMagicLeapPlugin.h
+++ b/Engine/Plugins/Lumin/MagicLeap/Source/MagicLeap/Public/IMagicLeapPlugin.h
@@ -36,7 +36,12 @@ public:
 		return FModuleManager::Get().IsModuleLoaded("MagicLeap");
 	}
 
-	virtual bool IsMagicLeapHMDValidAndConnected() const = 0;
+	/**
+	 * Checks to see if the XRSystem instance is a MagicLeap HMD.
+	 *
+	 * @return True if the current XRSystem instance is a MagicLeap HMD device
+	 */
+	virtual bool IsMagicLeapHMDValid() const = 0;
 
 	/**
 	* Returns the HMD associated with the ML plugin.
diff --git a/Engine/Plugins/Lumin/MagicLeap/Source/MagicLeap/Public/MagicLeapPrivileges.h b/Engine/Plugins/Lumin/MagicLeap/Source/MagicLeap/Public/MagicLeapPrivileges.h
new file mode 100644
--- /dev/null
+++ b/Engine/Plugins/Lumin/MagicLeap/Source/MagicLeap/Public/MagicLeapPrivileges.h
@@ -0,0 +1,135 @@
+// %BANNER_BEGIN%
+// ---------------------------------------------------------------------
+// %COPYRIGHT_BEGIN%
+//
+// Copyright (c) 2017 Magic Leap, Inc. (COMPANY) All Rights Reserved.
+// Magic Leap, Inc. Confidential and Proprietary
+//
+// NOTICE: All information contained herein is, and remains the property
+// of COMPANY. The intellectual and technical concepts contained herein
+// are proprietary to COMPANY and may be covered by U.S. and Foreign
+// Patents, patents in process, and are protected by trade secret or
+// copyright law. Dissemination of this information or reproduction of
+// this material is strictly forbidden unless prior written permission is
+// obtained from COMPANY. Access to the source code contained herein is
+// hereby forbidden to anyone except current COMPANY employees, managers
+// or contractors who have executed Confidentiality and Non-disclosure
+// agreements explicitly covering such access.
+//
+// The copyright notice above does not evidence any actual or intended
+// publication or disclosure of this source code, which includes
+// information that is confidential and/or proprietary, and is a trade
+// secret, of COMPANY. ANY REPRODUCTION, MODIFICATION, DISTRIBUTION,
+// PUBLIC PERFORMANCE, OR PUBLIC DISPLAY OF OR THROUGH USE OF THIS
+// SOURCE CODE WITHOUT THE EXPRESS WRITTEN CONSENT OF COMPANY IS
+// STRICTLY PROHIBITED, AND IN VIOLATION OF APPLICABLE LAWS AND
+// INTERNATIONAL TREATIES. THE RECEIPT OR POSSESSION OF THIS SOURCE
+// CODE AND/OR RELATED INFORMATION DOES NOT CONVEY OR IMPLY ANY RIGHTS
+// TO REPRODUCE, DISCLOSE OR DISTRIBUTE ITS CONTENTS, OR TO MANUFACTURE,
+// USE, OR SELL ANYTHING THAT IT MAY DESCRIBE, IN WHOLE OR IN PART.
+//
+// %COPYRIGHT_END%
+// --------------------------------------------------------------------
+// %BANNER_END%
+
+#pragma once
+
+#include "Components/ActorComponent.h"
+#include "MagicLeapPrivileges.generated.h"
+
+/** Priviliges an app can request for from the system. */
+UENUM(BlueprintType)
+enum class EMagicLeapPrivilege : uint8
+{
+	Invalid,
+	AudioRecognizer,
+	BatteryInfo,
+	CameraCapture,
+	WorldReconstruction,
+	InAppPurchase,
+	AudioCaptureMic,
+	DrmCertificates,
+	Occlusion,
+	LowLatencyLightwear,
+	Internet,
+	IdentityRead,
+	BackgroundDownload,
+	BackgroundUpload,
+	MediaDrm,
+	Media,
+	MediaMetadata,
+	PowerInfo,
+	LocalAreaNetwork,
+	VoiceInput,
+	Documents,
+	ConnectBackgroundMusicService,
+	RegisterBackgroundMusicService,
+	PwFoundObjRead,
+	NormalNotificationsUsage,
+	MusicService,
+	ControllerPose,
+	ScreensProvider,
+	GesturesSubscribe,
+	GesturesConfig,
+};
+
+/**
+ *  Class which provides functions to check and request the priviliges the app has at runtime.
+ */
+UCLASS(ClassGroup = MagicLeap, BlueprintType, Blueprintable, EditInlineNew, meta = (BlueprintSpawnableComponent))
+class MAGICLEAP_API UMagicLeapPrivileges : public UActorComponent
+{
+	GENERATED_BODY()
+
+public:
+	UMagicLeapPrivileges();
+	virtual ~UMagicLeapPrivileges();
+
+	/** Polls for and handles the results of the async privilege requests. */
+	virtual void TickComponent(float DeltaTime, enum ELevelTick TickType, FActorComponentTickFunction *ThisTickFunction) override;
+
+	/**
+	  Delegate for the result of requesting a privilege asynchronously.
+	  @param RequestedPrivilege The privilege that was requested.
+	  @param WasGranted True if the privilege was granted, false otherwise.
+	*/
+	DECLARE_DYNAMIC_DELEGATE_TwoParams(FPrivilegeRequestDelegate, EMagicLeapPrivilege, RequestedPrivilege, bool, WasGranted);
+
+	/**
+	  Check whether the application has the specified privilege.
+	  This does not solicit consent from the end-user and is non-blocking.
+	  @param Privilege The privilege to check.
+	  @return True if the privilege is granted, false otherwise.
+	*/
+	UFUNCTION(BlueprintCallable, Category = "Privileges|MagicLeap")
+	bool CheckPrivilege(EMagicLeapPrivilege Privilege);
+
+	/**
+	  Request the specified privilege.
+	  This may possibly solicit consent from the end-user; if so it will block.
+	  @param Privilege The privilege to request.
+	  @return True if the privilege is granted, false otherwise.
+	 */
+	UFUNCTION(BlueprintCallable, Category = "Privileges|MagicLeap")
+	bool RequestPrivilege(EMagicLeapPrivilege Privilege);
+
+	/**
+	  Request the specified privilege asynchronously.
+	  This may possibly solicit consent from the end-user. Result will be delivered
+	  to the specified delegate.
+	  @param Privilege The privilege to request.
+	  @param ResultDelegate Callback which reports the result of the request.
+	  @return True if the privilege request was successfully dispatched, false otherwise.
+	 */
+	UFUNCTION(BlueprintCallable, Category = "Privileges|MagicLeap")
+	bool RequestPrivilegeAsync(EMagicLeapPrivilege Privilege, const FPrivilegeRequestDelegate& ResultDelegate);
+
+private:
+	struct FPendingAsyncRequest
+	{
+		EMagicLeapPrivilege Privilege;
+		struct MLPrivilegesAsyncRequest* Request;
+		FPrivilegeRequestDelegate Delegate;
+	};
+	TArray<FPendingAsyncRequest> PendingAsyncRequests;
+};
diff --git a/Engine/Plugins/Lumin/MagicLeap/Source/MagicLeapAudio/Private/AudioMixerPlatformMagicLeap.cpp b/Engine/Plugins/Lumin/MagicLeap/Source/MagicLeapAudio/Private/AudioMixerPlatformMagicLeap.cpp
--- a/Engine/Plugins/Lumin/MagicLeap/Source/MagicLeapAudio/Private/AudioMixerPlatformMagicLeap.cpp
+++ b/Engine/Plugins/Lumin/MagicLeap/Source/MagicLeapAudio/Private/AudioMixerPlatformMagicLeap.cpp
@@ -16,7 +16,7 @@ DEFINE_LOG_CATEGORY(LogAudioMixerMagicLeap);
 
 // Macro to check result for failure, get string version, log, and return false
 #define MLAUDIO_RETURN_ON_FAIL(Result)						\
-	if (Result != MLAudioError_Success)						\
+	if (Result != MLResult_Ok)						\
 	{														\
 		const TCHAR* ErrorString = FMixerPlatformMagicLeap::GetErrorString(Result);	\
 		AUDIO_PLATFORM_ERROR(ErrorString);					\
@@ -24,7 +24,7 @@ DEFINE_LOG_CATEGORY(LogAudioMixerMagicLeap);
 	}
 
 #define MLAUDIO_CHECK_ON_FAIL(Result)						\
-	if (Result != MLAudioError_Success)						\
+	if (Result != MLResult_Ok)						\
 	{														\
 		const TCHAR* ErrorString = FMixerPlatformMagicLeap::GetErrorString(Result);	\
 		AUDIO_PLATFORM_ERROR(ErrorString);					\
@@ -32,7 +32,7 @@ DEFINE_LOG_CATEGORY(LogAudioMixerMagicLeap);
 	}
 
 #define MLAUDIO_LOG_ON_FAIL(Result)							\
-	if (Result != MLAudioError_Success)						\
+	if (Result != MLResult_Ok)						\
 	{														\
 		const TCHAR* ErrorString = FMixerPlatformMagicLeap::GetErrorString(Result);	\
 		UE_LOG(LogAudioMixerMagicLeap, Error, TEXT("Error in %s, line %s: %s"), __FILE__, __LINE__,ErrorString); \
@@ -42,27 +42,36 @@ namespace Audio
 {
 	// ML1 currently only has stereo speakers and stereo aux support.
 	const uint32 DefaultNumChannels = 2;
-	// TODO: @Epic check the value tp be used. Setting default for now.
+	// TODO: @Epic check the value to be used. Setting default for now.
 	const uint32 DefaultSamplesPerSecond = 48000;  // presumed 48KHz and 16 bits for the sample
+	// TODO: @Epic check the value to be used. Setting default for now.
+	const float DefaultMaxPitch = 1.0f;
 
 #if WITH_MLSDK
-	static void GetMLDeviceDefaults(MLAudioBufferFormat& OutFormat, uint32& OutSize, uint32& OutMinSize, MLAudioError& Result)
+	static void GetMLDeviceDefaults(MLAudioBufferFormat& OutFormat, uint32& OutSize, uint32& OutMinSize)
 	{
 		static bool bRetrievedDeviceDefaults = false;
 		static uint32 CachedSize = 0;
 		static uint32 CachedMinSize = 0;
 		static MLAudioBufferFormat CachedBufferFormat;
-		static MLAudioError CachedResult;
 
 		if (!bRetrievedDeviceDefaults)
 		{
-			bRetrievedDeviceDefaults = MLAudioGetOutputStreamDefaults(DefaultNumChannels, DefaultSamplesPerSecond, &CachedBufferFormat, &CachedSize, &CachedMinSize, &CachedResult);
+			MLResult Result = MLAudioGetOutputStreamDefaults(DefaultNumChannels, DefaultSamplesPerSecond, DefaultMaxPitch, &CachedBufferFormat, &CachedSize, &CachedMinSize);
+			if (Result == MLResult_Ok)
+			{
+				bRetrievedDeviceDefaults = true;
+			}
+			else
+			{
+				bRetrievedDeviceDefaults = false;
+				UE_LOG(LogAudioMixerMagicLeap, Error, TEXT("MLAudioGetOutputStreamDefaults failed with error %d."), Result);
+			}
 		}
 
 		OutFormat = CachedBufferFormat;
 		OutSize = CachedSize;
 		OutMinSize = CachedMinSize;
-		Result = CachedResult;
 	}
 #endif //WITH_MLSDK
 
@@ -86,28 +95,25 @@ namespace Audio
 	}
 
 #if WITH_MLSDK
-	const TCHAR* FMixerPlatformMagicLeap::GetErrorString(MLAudioError Result)
+	const TCHAR* FMixerPlatformMagicLeap::GetErrorString(MLResult Result)
 	{
 		switch (Result)
 		{
-		case MLAudioError_Success:					return TEXT("MLAudioError_Success");
-		case MLAudioError_NotImplemented:			return TEXT("MLAudioError_NotImplemented");
-		case MLAudioError_UnspecifiedFailure:       return TEXT("MLAudioError_UnspecifiedFailure");
-		case MLAudioError_UnauthorizedCommand:      return TEXT("MLAudioError_UnauthorizedCommand");
-		case MLAudioError_HandleNotFound:           return TEXT("MLAudioError_HandleNotFound");
-		case MLAudioError_InvalidArgument:          return TEXT("MLAudioError_InvalidArgument");
-		case MLAudioError_InvalidSampleRate:        return TEXT("MLAudioError_InvalidSampleRate");
-		case MLAudioError_InvalidBitsPerSample:     return TEXT("MLAudioError_InvalidBitsPerSample");
-		case MLAudioError_InvalidValidBits:         return TEXT("MLAudioError_InvalidValidBits");
-		case MLAudioError_InvalidSampleFormat:      return TEXT("MLAudioError_InvalidSampleFormat");
-		case MLAudioError_InvalidChannelCount:      return TEXT("MLAudioError_InvalidChannelCount");
-		case MLAudioError_InvalidChannelFormat:     return TEXT("MLAudioError_InvalidChannelFormat");
-		case MLAudioError_InvalidBufferSize:        return TEXT("MLAudioError_InvalidBufferSize");
-		case MLAudioError_BufferAllocFailure:       return TEXT("MLAudioError_BufferAllocFailure");
-		case MLAudioError_BufferNotReady:           return TEXT("MLAudioError_BufferNotReady");
-		case MLAudioError_FileNotFound:             return TEXT("MLAudioError_FileNotFound");
-		case MLAudioError_FileNotRecognized:        return TEXT("MLAudioError_FileNotRecognized");
-		default:                                return TEXT("MlAudioError_UnknownError");
+		case MLResult_UnspecifiedFailure:			 return TEXT("MLResult_UnspecifiedFailure");
+		case MLResult_InvalidParam:					 return TEXT("MLResult_InvalidParam");
+		case MLResult_AllocFailed:					 return TEXT("MLResult_AllocFailed");
+		case MLAudioResult_NotImplemented:			 return TEXT("MLAudioResult_NotImplemented");
+		case MLAudioResult_HandleNotFound:           return TEXT("MLAudioResult_HandleNotFound");
+		case MLAudioResult_InvalidSampleRate:        return TEXT("MLAudioResult_InvalidSampleRate");
+		case MLAudioResult_InvalidBitsPerSample:     return TEXT("MLAudioResult_InvalidBitsPerSample");
+		case MLAudioResult_InvalidValidBits:         return TEXT("MLAudioResult_InvalidValidBits");
+		case MLAudioResult_InvalidSampleFormat:      return TEXT("MLAudioResult_InvalidSampleFormat");
+		case MLAudioResult_InvalidChannelCount:      return TEXT("MLAudioResult_InvalidChannelCount");
+		case MLAudioResult_InvalidBufferSize:        return TEXT("MLAudioResult_InvalidBufferSize");
+		case MLAudioResult_BufferNotReady:           return TEXT("MLAudioResult_BufferNotReady");
+		case MLAudioResult_FileNotFound:             return TEXT("MLAudioResult_FileNotFound");
+		case MLAudioResult_FileNotRecognized:        return TEXT("MLAudioResult_FileNotRecognized");
+		default:                                return TEXT("MlAudioResult_UnknownError");
 		}
 	}
 #endif //WITH_MLSDK
@@ -158,12 +164,11 @@ namespace Audio
 	bool FMixerPlatformMagicLeap::GetOutputDeviceInfo(const uint32 InDeviceIndex, FAudioPlatformDeviceInfo& OutInfo)
 	{
 #if WITH_MLSDK
-		MLAudioError Result = MLAudioError_Success;
 		MLAudioBufferFormat DesiredBufferFormat;
 		uint32 OutSize;
 		uint32 OutMinimalSize;
 
-		GetMLDeviceDefaults(DesiredBufferFormat, OutSize, OutMinimalSize, Result);
+		GetMLDeviceDefaults(DesiredBufferFormat, OutSize, OutMinimalSize);
 
 		OutInfo.Name = TEXT("Magic Leap Audio Device");
 		OutInfo.DeviceId = 0;
@@ -213,12 +218,11 @@ namespace Audio
 			return false;
 		}
 
-		MLAudioError Result = MLAudioError_Success;
 		MLAudioBufferFormat DesiredBufferFormat;
 		uint32 OutSize;
 		uint32 OutMinimalSize;
 
-		GetMLDeviceDefaults(DesiredBufferFormat, OutSize, OutMinimalSize, Result);
+		GetMLDeviceDefaults(DesiredBufferFormat, OutSize, OutMinimalSize);
 
 		OpenStreamParams = Params;
 
@@ -236,7 +240,9 @@ namespace Audio
 
 		DesiredBufferFormat.channel_count = DefaultNumChannels;
 
-		if (!MLAudioCreateSoundWithOutputStream(&DesiredBufferFormat, OutSize, &MLAudioCallback, this, &StreamHandle, &Result))
+		MLResult Result = MLAudioCreateSoundWithOutputStream(&DesiredBufferFormat, OutSize, &MLAudioCallback, this, &StreamHandle);
+
+		if (Result != MLResult_Ok)
 		{
 			MLAUDIO_LOG_ON_FAIL(Result);
 			return false;
@@ -258,9 +264,9 @@ namespace Audio
 			return false;
 		}
 
-		MLAudioError Result = MLAudioError_Success;
+		MLResult Result = MLAudioDestroySound(StreamHandle);
 
-		if (!MLAudioDestroySound(StreamHandle, &Result))
+		if (Result != MLResult_Ok)
 		{
 			MLAUDIO_LOG_ON_FAIL(Result);
 			return false;
@@ -280,30 +286,31 @@ namespace Audio
 
 		check(MLHandleIsValid(StreamHandle));
 
-		MLAudioError Result = MLAudioError_Success;
+		MLResult Result = MLResult_Ok;
 
 		//Pre buffer with two zeroed buffers:
 		static const int32 NumberOfBuffersToPrecache = 2;
 		for (int32 BufferIndex = 0; BufferIndex < NumberOfBuffersToPrecache; BufferIndex++)
 		{
 			MLAudioBuffer PrecacheBuffer;
-
-			if (!MLAudioGetOutputStreamBuffer(StreamHandle, &PrecacheBuffer, &Result))
+			Result = MLAudioGetOutputStreamBuffer(StreamHandle, &PrecacheBuffer);
+			if (Result != MLResult_Ok)
 			{
 				MLAUDIO_LOG_ON_FAIL(Result);
 				break;
 			}
 
 			FMemory::Memzero(PrecacheBuffer.ptr, PrecacheBuffer.size);
-
-			if (!MLAudioReleaseOutputStreamBuffer(StreamHandle, &Result))
+			Result = MLAudioReleaseOutputStreamBuffer(StreamHandle);
+			if (Result != MLResult_Ok)
 			{
 				MLAUDIO_LOG_ON_FAIL(Result);
 				break;
 			}
 		}
 
-		if (!MLAudioStartSound(StreamHandle, &Result))
+		Result = MLAudioStartSound(StreamHandle);
+		if (Result != MLResult_Ok)
 		{
 			MLAUDIO_LOG_ON_FAIL(Result);
 			return false;
@@ -316,13 +323,13 @@ namespace Audio
 	bool FMixerPlatformMagicLeap::StopAudioStream()
 	{
 #if WITH_MLSDK
-		MLAudioError Result;
 		if (!bInitialized || AudioStreamInfo.StreamState != EAudioOutputStreamState::Running)
 		{
 			return false;
 		}
 
-		if (!MLAudioStopSound(StreamHandle, &Result))
+		MLResult Result = MLAudioStopSound(StreamHandle);
+		if (Result != MLResult_Ok)
 		{
 			MLAUDIO_LOG_ON_FAIL(Result);
 			return false;
@@ -348,12 +355,11 @@ namespace Audio
 	FAudioPlatformSettings FMixerPlatformMagicLeap::GetPlatformSettings() const
 	{
 #if WITH_MLSDK
-		MLAudioError Result = MLAudioError_Success;
 		MLAudioBufferFormat DesiredBufferFormat;
 		uint32 OutSize;
 		uint32 OutMinimalSize;
 
-		GetMLDeviceDefaults(DesiredBufferFormat, OutSize, OutMinimalSize, Result);
+		GetMLDeviceDefaults(DesiredBufferFormat, OutSize, OutMinimalSize);
 
 		FAudioPlatformSettings PlatformSettings;
 		PlatformSettings.CallbackBufferFrameSize = OutSize / (DefaultNumChannels * (DesiredBufferFormat.bits_per_sample / 8));
@@ -417,9 +423,8 @@ namespace Audio
 		{
 			if (MLHandleIsValid(StreamHandle))
 			{
-				MLAudioError Result;
-
-				if (!MLAudioStartSound(StreamHandle, &Result))
+				MLResult Result = MLAudioStartSound(StreamHandle);
+				if (Result != MLResult_Ok)
 				{
 					MLAUDIO_LOG_ON_FAIL(Result);
 					return;
@@ -434,12 +439,11 @@ namespace Audio
 	int32 FMixerPlatformMagicLeap::GetNumFrames(const int32 InNumReqestedFrames)
 	{
 #if WITH_MLSDK
-		MLAudioError Result = MLAudioError_Success;
 		MLAudioBufferFormat DesiredBufferFormat;
 		uint32 OutSize;
 		uint32 OutMinimalSize;
 
-		GetMLDeviceDefaults(DesiredBufferFormat, OutSize, OutMinimalSize, Result);
+		GetMLDeviceDefaults(DesiredBufferFormat, OutSize, OutMinimalSize);
 		return OutSize / (DefaultNumChannels * (DesiredBufferFormat.bits_per_sample / 8));
 #else
 		return 0;
@@ -453,9 +457,8 @@ namespace Audio
 		{
 			if (MLHandleIsValid(StreamHandle))
 			{
-				MLAudioError Result;
-
-				if (!MLAudioStopSound(StreamHandle, &Result))
+				MLResult Result = MLAudioStopSound(StreamHandle);
+				if (Result != MLResult_Ok)
 				{
 					MLAUDIO_LOG_ON_FAIL(Result);
 					return;
@@ -467,18 +470,22 @@ namespace Audio
 #endif //WITH_MLSDK
 	}
 
-	void FMixerPlatformMagicLeap::MLAudioCallback(void* CallbackContext)
-	{
 #if WITH_MLSDK
+	void FMixerPlatformMagicLeap::MLAudioCallback(MLHandle Handle, void* CallbackContext)
+	{
 		FMixerPlatformMagicLeap* InPlatform = (FMixerPlatformMagicLeap*)CallbackContext;
-		MLAudioError Result;
 
 		MLAudioBuffer CallbackBuffer;
 
 		check(MLHandleIsValid(InPlatform->StreamHandle));
-
+		// TODO: should this be an assert?
+		if (Handle != InPlatform->StreamHandle)
+		{
+			return;
+		}
 		// Get the callback buffer from MLAudio 
-		if (!MLAudioGetOutputStreamBuffer(InPlatform->StreamHandle, &CallbackBuffer, &Result))
+		MLResult Result = MLAudioGetOutputStreamBuffer(InPlatform->StreamHandle, &CallbackBuffer);
+		if (Result != MLResult_Ok)
 		{
 			MLAUDIO_LOG_ON_FAIL(Result);
 			return;
@@ -497,14 +504,14 @@ namespace Audio
 			FMemory::Memcpy(CallbackBuffer.ptr, InPlatform->CachedBufferHandle, CallbackBuffer.size);
 		}
 
-		if (!MLAudioReleaseOutputStreamBuffer(InPlatform->StreamHandle, &Result))
+		Result = MLAudioReleaseOutputStreamBuffer(InPlatform->StreamHandle);
+		if (Result != MLResult_Ok)
 		{
 			MLAUDIO_LOG_ON_FAIL(Result);
 			return;
 		}
 
 		InPlatform->CachedBufferHandle = nullptr;
-#endif //WITH_MLSDK
 	}
-
+#endif //WITH_MLSDK
 }
diff --git a/Engine/Plugins/Lumin/MagicLeap/Source/MagicLeapAudio/Private/AudioMixerPlatformMagicLeap.h b/Engine/Plugins/Lumin/MagicLeap/Source/MagicLeapAudio/Private/AudioMixerPlatformMagicLeap.h
--- a/Engine/Plugins/Lumin/MagicLeap/Source/MagicLeapAudio/Private/AudioMixerPlatformMagicLeap.h
+++ b/Engine/Plugins/Lumin/MagicLeap/Source/MagicLeapAudio/Private/AudioMixerPlatformMagicLeap.h
@@ -54,7 +54,7 @@ namespace Audio
 
 	private:
 #if WITH_MLSDK
-		static const TCHAR* GetErrorString(MLAudioError Result);
+		static const TCHAR* GetErrorString(MLResult Result);
 #endif //WITH_MLSDK
 
 		bool bSuspended;
@@ -63,10 +63,10 @@ namespace Audio
 
 #if WITH_MLSDK
 		MLHandle StreamHandle;
-#endif //WITH_MLSDK
 
 		// Static callback used for MLAudio:
-		static void MLAudioCallback(void* CallbackContext);
+		static void MLAudioCallback(MLHandle Handle, void* CallbackContext);
+#endif //WITH_MLSDK
 	};
 
 }
diff --git a/Engine/Plugins/Lumin/MagicLeap/Source/MagicLeapController/Private/MagicLeapController.cpp b/Engine/Plugins/Lumin/MagicLeap/Source/MagicLeapController/Private/MagicLeapController.cpp
--- a/Engine/Plugins/Lumin/MagicLeap/Source/MagicLeapController/Private/MagicLeapController.cpp
+++ b/Engine/Plugins/Lumin/MagicLeap/Source/MagicLeapController/Private/MagicLeapController.cpp
@@ -560,10 +560,39 @@ void FMagicLeapController::Enable()
 #if WITH_MLSDK
 	if (!MLHandleIsValid(InputTracker))
 	{
-		InputTracker = MLInputCreate(nullptr);
-		FMemory::Memset(&InputState, 0, sizeof(InputState));
-		MLInputSetControllerCallbacks(InputTracker, &InputControllerCallbacks, this);
-		MLInputSetKeyboardCallbacks(InputTracker, &InputKeyboardCallbacks, this);
+		MLResult Result = MLInputCreate(nullptr, &InputTracker);
+		if (Result == MLResult_Ok)
+		{
+			FMemory::Memset(&InputState, 0, sizeof(InputState));
+			Result = MLInputSetControllerCallbacks(InputTracker, &InputControllerCallbacks, this);
+			if (Result == MLResult_Ok)
+			{
+				Result = MLInputSetKeyboardCallbacks(InputTracker, &InputKeyboardCallbacks, this);
+				if (Result == MLResult_Ok)
+				{
+					// Poll to get startup status
+					UpdateTrackerData();
+				}
+				else
+				{
+					UE_LOG(LogMagicLeapController, Error, 
+						TEXT("MLInputSetKeyboardCallbacks failed with error %s."), 
+						UTF8_TO_TCHAR(MLGetResultString(Result)));
+				}
+			}
+			else
+			{
+				UE_LOG(LogMagicLeapController, Error, 
+					TEXT("MLInputSetControllerCallbacks failed with error %s."), 
+					UTF8_TO_TCHAR(MLGetResultString(Result)));
+			}
+		}
+		else
+		{
+			UE_LOG(LogMagicLeapController, Error, 
+				TEXT("MLInputCreate failed with error %s."), 
+				UTF8_TO_TCHAR(MLGetResultString(Result)));
+		}
 	}
 #endif //WITH_MLSDK
 }
@@ -578,15 +607,15 @@ void FMagicLeapController::Disable()
 #if WITH_MLSDK
 	if (MLHandleIsValid(InputTracker))
 	{
-		if (MLInputDestroy(InputTracker))
+		MLResult Result = MLInputDestroy(InputTracker);
+		if (Result != MLResult_Ok)
 		{
-			InputTracker = ML_INVALID_HANDLE;
-			bIsInputStateValid = false;
-		}
-		else
-		{
-			UE_LOG(LogMagicLeapController, Error, TEXT("Error destroying input tracker."));
+			UE_LOG(LogMagicLeapController, Error, 
+				TEXT("MLInputDestroy failed with error %s!"), 
+				UTF8_TO_TCHAR(MLGetResultString(Result)));
 		}
+		InputTracker = ML_INVALID_HANDLE;
+		bIsInputStateValid = false;
 	}
 #endif //WITH_MLSDK
 }
@@ -642,25 +671,23 @@ void FMagicLeapController::InternalSetHapticFeedbackValues(int32 ControllerId, i
 
 					UE_LOG(LogMagicLeapController, Warning, TEXT("Disabling haptic feedback pending a lower level API.  please use UMagicLeapControllerFunctionLibrary::PlayControllerHapticFeedback for now!"));
 
-					bool bResult = false;
 					if ((EControllerHand(Hand) == EControllerHand::Left))
 					{
 						// TODO: What to do with duration?
-						//bResult = MLInputHapticsStartControllerBody(InputTracker, static_cast<uint8>(*ControllerIndex), MLInputControllerFeedbackPatternVibe_Buzz, intensity, 8);
-						if (!bResult)
-						{
-							UE_LOG(LogMagicLeapController, Error, TEXT("MLInputHapticsStartControllerBody() failed."));
-						}
+						//MLResult Result = MLInputHapticsStartControllerBody(InputTracker, static_cast<uint8>(*ControllerIndex), MLInputControllerFeedbackPatternVibe_Buzz, intensity, 8);
+						//if (Result != MLResult_Ok)
+						//{
+						//	UE_LOG(LogMagicLeapController, Error, TEXT("MLInputHapticsStartControllerBody() failed."));
+						//}
 					}
 					else
 					{
 						// TODO: What to do with duration?
-						//bResult = MLInputHapticsStartControllerTouchpad(InputTracker, static_cast<uint8>(*Controller), MLInputControllerFeedbackPatternVibe_Buzz, intensity, 8);
-						if (!bResult)
-						{
-							UE_LOG(LogMagicLeapController, Error, TEXT("MLInputHapticsStartControllerTouchpad() failed."));
-						}
-
+						//MLResult Result = MLInputHapticsStartControllerTouchpad(InputTracker, static_cast<uint8>(*Controller), MLInputControllerFeedbackPatternVibe_Buzz, intensity, 8);
+						//if (Result != MLResult_Ok)
+						//{
+						//	UE_LOG(LogMagicLeapController, Error, TEXT("MLInputHapticsStartControllerTouchpad() failed."));
+						//}
 					}
 
 					//only toggle when called from the haptics interface, not force feedback
@@ -704,26 +731,14 @@ bool FMagicLeapController::GetControllerOrientationAndPosition(const int32 Contr
 		if (GetControllerTrackingStatus(ControllerIndex, DeviceHand) != ETrackingStatus::NotTracked)
 		{
 			const FTransform* ControllerTransform = &LeftControllerTransform;
-			const FTransform* CalibrationTransform = &LeftControllerCalibration;
-			ETrackingStatus ControllerTrackingStatus = GetControllerTrackingStatus(DeviceIndex, EControllerHand::Left);
-			EMLControllerType bControllerType = GetMLControllerType(EControllerHand::Left);
 
 			if (DeviceHand == EControllerHand::Right)
 			{
-				ControllerTrackingStatus = GetControllerTrackingStatus(DeviceIndex, EControllerHand::Right);
 				ControllerTransform = &RightControllerTransform;
-				CalibrationTransform = &RightControllerCalibration;
 			}
 
-			OutPosition = ControllerTransform->GetLocation() + CalibrationTransform->GetLocation();
-			if (ControllerTrackingStatus == ETrackingStatus::InertialOnly)
-			{
-				OutOrientation = (CalibrationTransform->GetRotation() * ControllerTransform->GetRotation()).Rotator();
-			}
-			else
-			{
-				OutOrientation = (ControllerTransform->GetRotation() * CalibrationTransform->GetRotation()).Rotator();
-			}
+			OutPosition = ControllerTransform->GetLocation();
+			OutOrientation = ControllerTransform->GetRotation().Rotator();
 
 			bControllerTracked = true;
 		}
@@ -771,11 +786,11 @@ FName FMagicLeapController::GetMotionControllerDeviceTypeName() const
 void FMagicLeapController::UpdateTrackerData()
 {
 #if WITH_MLSDK
-	if (MLHandleIsValid(InputTracker))
+	if (MLHandleIsValid(InputTracker) && IMagicLeapPlugin::Get().IsMagicLeapHMDValid())
 	{
 		// Memcopy instead of assignment operator since its an array type of a c-struct and we need a deep copy here.
 		FMemory::Memcpy(&OldInputState, &InputState, sizeof(InputState));
-		bIsInputStateValid = MLInputGetControllerState(InputTracker, InputState);
+		bIsInputStateValid = MLInputGetControllerState(InputTracker, InputState) == MLResult_Ok;
 
 		if (bIsInputStateValid)
 		{
@@ -874,7 +889,7 @@ void FMagicLeapController::InvertControllerMapping()
 EMLControllerType FMagicLeapController::GetMLControllerType(EControllerHand Hand) const
 {
 #if WITH_MLSDK
-	const int32* ControllerID = HandToControllerID.Find(EControllerHand::Left);
+	const int32* ControllerID = HandToControllerID.Find(Hand);
 	if (ControllerID != nullptr)
 	{
 		switch (InputState[*ControllerID].type)
@@ -889,24 +904,6 @@ EMLControllerType FMagicLeapController::GetMLControllerType(EControllerHand Hand
 	return EMLControllerType::None;
 }
 
-void FMagicLeapController::CalibrateControllerNow(EControllerHand Hand, const FVector& StartPosition, const FRotator& StartOrientation)
-{
-	const int32* ControllerID = HandToControllerID.Find(Hand);
-	if (ControllerID != nullptr)
-	{
-		FTransform* CalibrationTransform = &LeftControllerCalibration;
-		FTransform* CurrentTransform = &LeftControllerTransform;
-		if (Hand == EControllerHand::Right)
-		{
-			CalibrationTransform = &RightControllerCalibration;
-			CurrentTransform = &RightControllerTransform;
-		}
-
-		CalibrationTransform->SetRotation(StartOrientation.Quaternion() * CurrentTransform->GetRotation().Inverse());
-		CalibrationTransform->SetLocation(StartPosition - CurrentTransform->GetLocation());
-	}
-}
-
 bool FMagicLeapController::PlayControllerLED(EControllerHand Hand, EMLControllerLEDPattern LEDPattern, EMLControllerLEDColor LEDColor, float DurationInSec)
 {
 #if WITH_MLSDK
@@ -915,7 +912,7 @@ bool FMagicLeapController::PlayControllerLED(EControllerHand Hand, EMLController
 		const int32* ControllerID = HandToControllerID.Find(Hand);
 		if (ControllerID != nullptr)
 		{
-			return MLInputStartControllerFeedbackPatternLED(InputTracker, static_cast<uint8>(*ControllerID), UnrealToMLPatternLED(LEDPattern), UnrealToMLColorLED(LEDColor), static_cast<uint32>(DurationInSec * 1000));
+			return MLInputStartControllerFeedbackPatternLED(InputTracker, static_cast<uint8>(*ControllerID), UnrealToMLPatternLED(LEDPattern), UnrealToMLColorLED(LEDColor), static_cast<uint32>(DurationInSec * 1000)) == MLResult_Ok;
 		}
 	}
 	else
@@ -935,7 +932,7 @@ bool FMagicLeapController::PlayControllerLEDEffect(EControllerHand Hand, EMLCont
 		const int32* ControllerID = HandToControllerID.Find(Hand);
 		if (ControllerID != nullptr)
 		{
-			return MLInputStartControllerFeedbackPatternEffectLED(InputTracker, static_cast<uint8>(*ControllerID), UnrealToMLEffectLED(LEDEffect), UnrealToMLSpeedLED(LEDSpeed), UnrealToMLPatternLED(LEDPattern), UnrealToMLColorLED(LEDColor), static_cast<uint32>(DurationInSec * 1000));
+			return MLInputStartControllerFeedbackPatternEffectLED(InputTracker, static_cast<uint8>(*ControllerID), UnrealToMLEffectLED(LEDEffect), UnrealToMLSpeedLED(LEDSpeed), UnrealToMLPatternLED(LEDPattern), UnrealToMLColorLED(LEDColor), static_cast<uint32>(DurationInSec * 1000)) == MLResult_Ok;
 		}
 	}
 	else
@@ -955,7 +952,7 @@ bool FMagicLeapController::PlayControllerHapticFeedback(EControllerHand Hand, EM
 		const int32* ControllerID = HandToControllerID.Find(Hand);
 		if (ControllerID != nullptr)
 		{
-			return MLInputStartControllerFeedbackPatternVibe(InputTracker, static_cast<uint8>(*ControllerID), UnrealToMLPatternVibe(HapticPattern), UnrealToMLHapticIntensity(Intensity));
+			return MLInputStartControllerFeedbackPatternVibe(InputTracker, static_cast<uint8>(*ControllerID), UnrealToMLPatternVibe(HapticPattern), UnrealToMLHapticIntensity(Intensity)) == MLResult_Ok;
 		}
 	}
 	else
@@ -1138,6 +1135,8 @@ MLInputControllerFeedbackPatternLED FMagicLeapController::UnrealToMLPatternLED(E
 {
 	switch (LEDPattern)
 	{
+		case EMLControllerLEDPattern::None:
+			return MLInputControllerFeedbackPatternLED_None;
 		case EMLControllerLEDPattern::Clock01:
 			return MLInputControllerFeedbackPatternLED_Clock1;
 		case EMLControllerLEDPattern::Clock02:
@@ -1253,6 +1252,8 @@ MLInputControllerFeedbackPatternVibe FMagicLeapController::UnrealToMLPatternVibe
 {
 	switch (HapticPattern)
 	{
+		case EMLControllerHapticPattern::None:
+			return MLInputControllerFeedbackPatternVibe_None;
 		case EMLControllerHapticPattern::Click:
 			return MLInputControllerFeedbackPatternVibe_Click;
 		case EMLControllerHapticPattern::Bump:
diff --git a/Engine/Plugins/Lumin/MagicLeap/Source/MagicLeapController/Private/MagicLeapController.h b/Engine/Plugins/Lumin/MagicLeap/Source/MagicLeapController/Private/MagicLeapController.h
--- a/Engine/Plugins/Lumin/MagicLeap/Source/MagicLeapController/Private/MagicLeapController.h
+++ b/Engine/Plugins/Lumin/MagicLeap/Source/MagicLeapController/Private/MagicLeapController.h
@@ -83,7 +83,6 @@ public:
 	bool GetControllerMapping(int32 ControllerIndex, EControllerHand& Hand) const;
 	void InvertControllerMapping();
 	EMLControllerType GetMLControllerType(EControllerHand Hand) const;
-	void CalibrateControllerNow(EControllerHand Hand, const FVector& StartPosition, const FRotator& StartOrientation);
 
 	bool PlayControllerLED(EControllerHand Hand, EMLControllerLEDPattern LEDPattern, EMLControllerLEDColor LEDColor, float DurationInSec);
 	bool PlayControllerLEDEffect(EControllerHand Hand, EMLControllerLEDEffect LEDEffect, EMLControllerLEDSpeed LEDSpeed, EMLControllerLEDPattern LEDPattern, EMLControllerLEDColor LEDColor, float DurationInSec);
@@ -135,9 +134,6 @@ private:
 	FTransform LeftControllerTransform;
 	FTransform RightControllerTransform;
 
-	FTransform LeftControllerCalibration;
-	FTransform RightControllerCalibration;
-
 	TMap<int32, EControllerHand> ControllerIDToHand;
 	TMap<EControllerHand, int32> HandToControllerID;
 
diff --git a/Engine/Plugins/Lumin/MagicLeap/Source/MagicLeapController/Private/MagicLeapControllerFunctionLibrary.cpp b/Engine/Plugins/Lumin/MagicLeap/Source/MagicLeapController/Private/MagicLeapControllerFunctionLibrary.cpp
--- a/Engine/Plugins/Lumin/MagicLeap/Source/MagicLeapController/Private/MagicLeapControllerFunctionLibrary.cpp
+++ b/Engine/Plugins/Lumin/MagicLeap/Source/MagicLeapController/Private/MagicLeapControllerFunctionLibrary.cpp
@@ -43,15 +43,6 @@ EMLControllerType UMagicLeapControllerFunctionLibrary::GetMLControllerType(ECont
 	return EMLControllerType::None;
 }
 
-void UMagicLeapControllerFunctionLibrary::CalibrateControllerNow(EControllerHand Hand, const FVector& StartPosition, const FRotator& StartOrientation)
-{
-	TSharedPtr<FMagicLeapController> controller = StaticCastSharedPtr<FMagicLeapController>(IMagicLeapControllerPlugin::Get().GetInputDevice());
-	if (controller.IsValid())
-	{
-		controller->CalibrateControllerNow(Hand, StartPosition, StartOrientation);
-	}
-}
-
 bool UMagicLeapControllerFunctionLibrary::PlayControllerLED(EControllerHand Hand, EMLControllerLEDPattern LEDPattern, EMLControllerLEDColor LEDColor, float DurationInSec)
 {
 	TSharedPtr<FMagicLeapController> controller = StaticCastSharedPtr<FMagicLeapController>(IMagicLeapControllerPlugin::Get().GetInputDevice());
diff --git a/Engine/Plugins/Lumin/MagicLeap/Source/MagicLeapController/Public/MagicLeapControllerFunctionLibrary.h b/Engine/Plugins/Lumin/MagicLeap/Source/MagicLeapController/Public/MagicLeapControllerFunctionLibrary.h
--- a/Engine/Plugins/Lumin/MagicLeap/Source/MagicLeapController/Public/MagicLeapControllerFunctionLibrary.h
+++ b/Engine/Plugins/Lumin/MagicLeap/Source/MagicLeapController/Public/MagicLeapControllerFunctionLibrary.h
@@ -51,14 +51,6 @@ public:
 	UFUNCTION(BlueprintCallable, BlueprintPure, Category = "MotionController|MagicLeap")
 	static EMLControllerType GetMLControllerType(EControllerHand Hand);
 
-	/**
-	  Match the position & orientation of the physical controller with an entity in the map and call this function
-	  with the position & orientation of that entity relative to the Player Pawn. This would apply the correct offsets
-	  to the MotionController component's transform.
-	*/
-	UFUNCTION(BlueprintCallable, Category = "MotionController|MagicLeap")
-	static void CalibrateControllerNow(EControllerHand Hand, const FVector& StartPosition, const FRotator& StartOrientation);
-
 	/**
 	  Light up the LED on the Magic Leap Controller in the given pattern for the specified duration.
 	  @param Hand Controller to play the LED pattern on.
diff --git a/Engine/Plugins/Lumin/MagicLeap/Source/MagicLeapController/Public/MagicLeapControllerKeys.h b/Engine/Plugins/Lumin/MagicLeap/Source/MagicLeapController/Public/MagicLeapControllerKeys.h
--- a/Engine/Plugins/Lumin/MagicLeap/Source/MagicLeapController/Public/MagicLeapControllerKeys.h
+++ b/Engine/Plugins/Lumin/MagicLeap/Source/MagicLeapController/Public/MagicLeapControllerKeys.h
@@ -31,6 +31,7 @@ enum class EMLControllerType : uint8
 UENUM(BlueprintType)
 enum class EMLControllerLEDPattern : uint8
 {
+	None,
 	Clock01,
 	Clock02,
 	Clock03,
@@ -98,6 +99,7 @@ enum class EMLControllerLEDSpeed : uint8
 UENUM(BlueprintType)
 enum class EMLControllerHapticPattern : uint8
 {
+	None,
 	Click,
 	Bump,
 	DoubleClick,
diff --git a/Engine/Plugins/Lumin/MagicLeap/Source/MagicLeapEyeTracker/Private/MagicLeapEyeTrackerModule.cpp b/Engine/Plugins/Lumin/MagicLeap/Source/MagicLeapEyeTracker/Private/MagicLeapEyeTrackerModule.cpp
--- a/Engine/Plugins/Lumin/MagicLeap/Source/MagicLeapEyeTracker/Private/MagicLeapEyeTrackerModule.cpp
+++ b/Engine/Plugins/Lumin/MagicLeap/Source/MagicLeapEyeTracker/Private/MagicLeapEyeTrackerModule.cpp
@@ -144,6 +144,23 @@ bool FMagicLeapEyeTracker::GetEyeBlinkState(FMagicLeapEyeBlinkState& BlinkState)
 	return false;
 }
 
+bool FMagicLeapEyeTracker::GetFixationComfort(FMagicLeapFixationComfort& FixationComfort) const
+{
+	if (VREyeTracker)
+	{
+		const FMagicLeapVREyeTrackingData& VRGazeData = VREyeTracker->GetVREyeTrackingData();
+		if (EEyeTrackerStatus::Tracking == GetEyeTrackerStatus())
+		{
+			FixationComfort.FixationDepthIsUncomfortable = VRGazeData.FixationDepthIsUncomfortable;
+			FixationComfort.FixationDepthViolationHasOccurred = VRGazeData.FixationDepthViolationHasOccurred;
+			FixationComfort.RemainingTimeAtUncomfortableDepth = VRGazeData.RemainingTimeAtUncomfortableDepth;
+
+			return true;
+		}
+	}
+
+	return false;
+}
 
 /************************************************************************/
 /* FMagicLeapEyeTrackerModule                                                     */
@@ -218,3 +235,15 @@ bool UMagicLeapEyeTrackerFunctionLibrary::GetEyeBlinkState(FMagicLeapEyeBlinkSta
 
 	return false;
 }
+
+bool UMagicLeapEyeTrackerFunctionLibrary::GetFixationComfort(FMagicLeapFixationComfort& FixationComfort)
+{
+	// TODO: Don't do this. Use StaticCastSharedPtr().
+	FMagicLeapEyeTracker* ET = GEngine ? static_cast<FMagicLeapEyeTracker*>(GEngine->EyeTrackingDevice.Get()) : nullptr;
+	if (ET)
+	{
+		return ET->GetFixationComfort(FixationComfort);
+	}
+
+	return false;
+}
