#DEV-VR - Merge from //UE4/Dev-MagicLeap/... @ CL 4136411

diff --git a/Engine/Plugins/Lumin/MagicLeap/Source/MagicLeapEyeTracker/Private/MagicLeapVREyeTracker.cpp b/Engine/Plugins/Lumin/MagicLeap/Source/MagicLeapEyeTracker/Private/MagicLeapVREyeTracker.cpp
--- a/Engine/Plugins/Lumin/MagicLeap/Source/MagicLeapEyeTracker/Private/MagicLeapVREyeTracker.cpp
+++ b/Engine/Plugins/Lumin/MagicLeap/Source/MagicLeapEyeTracker/Private/MagicLeapVREyeTracker.cpp
@@ -10,12 +10,14 @@
 #include "GameFramework/PlayerController.h"
 #include "Engine/LocalPlayer.h"
 #include "SceneView.h"
+#include "IMagicLeapPlugin.h"
 #include "AppFramework.h"
 #include "MagicLeapHMD.h"
 #include "UnrealEngine.h"
 #include "MagicLeapSDKDetection.h"
 #include "HeadMountedDisplayFunctionLibrary.h"
 #include "MagicLeapPluginUtil.h" // for ML_INCLUDES_START/END
+#include "IMagicLeapPlugin.h"
 
 #if WITH_MLSDK
 ML_INCLUDES_START
@@ -40,7 +42,10 @@ FMagicLeapVREyeTracker::~FMagicLeapVREyeTracker()
 #if WITH_MLSDK
 	if (EyeTrackingHandle != ML_INVALID_HANDLE)
 	{
-		MLEyeTrackingDestroy(EyeTrackingHandle);
+		if (MLEyeTrackingDestroy(EyeTrackingHandle) != MLResult_Ok)
+		{
+			UE_LOG(LogCore, Warning, TEXT("   MLEyeTrackingDestroy failure"));
+		}
 		EyeTrackingHandle = ML_INVALID_HANDLE;
 	}
 #endif //WITH_MLSDK
@@ -67,16 +72,16 @@ void FMagicLeapVREyeTracker::SetActivePlayerController(APlayerController* NewAct
 bool FMagicLeapVREyeTracker::Tick(float DeltaTime)
 {
 	bool bSuccess = true;
-#if WITH_MLSDK && PLATFORM_LUMIN
+#if WITH_MLSDK
 	//assume we're in a bad state
 	UnfilteredEyeTrackingData.bIsStable = false;
 
-	if (EyeTrackingHandle != ML_INVALID_HANDLE)
+	if (MLHandleIsValid(EyeTrackingHandle))
 	{
 		//check stat first to make sure everything is valid
 		MLEyeTrackingState TempTrackingState;
 		FMemory::Memzero(TempTrackingState);
-		bSuccess = bSuccess && MLEyeTrackingGetState(EyeTrackingHandle, &TempTrackingState);
+		bSuccess = bSuccess && (MLEyeTrackingGetState(EyeTrackingHandle, &TempTrackingState) == MLResult_Ok);
 
 		//make sure this is valid eye tracking data!
 		if (bSuccess
@@ -87,7 +92,8 @@ bool FMagicLeapVREyeTracker::Tick(float DeltaTime)
 		{
 			EyeTrackingStatus = EMagicLeapEyeTrackingStatus::UserPresentAndWatchingWindow;
 
-			bIsCalibrated = TempTrackingState.calibration_complete;
+			// TODO: expose as enum instead of simple enum.
+			bIsCalibrated = TempTrackingState.calibration_status == MLEyeTrackingCalibrationStatus_Good;
 
 			UnfilteredEyeTrackingData.bIsStable = true;
 			FDateTime Now = FDateTime::UtcNow();
@@ -101,36 +107,54 @@ bool FMagicLeapVREyeTracker::Tick(float DeltaTime)
 				//harvest the 3 transforms
 				EFailReason FailReason;
 				FTransform FixationTransform;
-				AppFramework.GetTransform(EyeTrackingStaticData.fixation, FixationTransform, FailReason);
-				FixationTransform.AddToTranslation(PoseTransform.GetLocation());
-				FixationTransform.ConcatenateRotation(PoseTransform.Rotator().Quaternion());
+				if (AppFramework.GetTransform(EyeTrackingStaticData.fixation, FixationTransform, FailReason))
+				{
+					FixationTransform.AddToTranslation(PoseTransform.GetLocation());
+					FixationTransform.ConcatenateRotation(PoseTransform.Rotator().Quaternion());
 
-				FTransform LeftCenterTransform;
-				AppFramework.GetTransform(EyeTrackingStaticData.left_center, LeftCenterTransform, FailReason);
-				LeftCenterTransform.AddToTranslation(PoseTransform.GetLocation());
-				LeftCenterTransform.ConcatenateRotation(PoseTransform.Rotator().Quaternion());
+					//get focal point
+					UnfilteredEyeTrackingData.WorldAverageGazeConvergencePoint = FixationTransform.GetTranslation();
+				}
 
+				bool bLeftTransformValid = false;
+				FTransform LeftCenterTransform;
+				if (AppFramework.GetTransform(EyeTrackingStaticData.left_center, LeftCenterTransform, FailReason))
+				{
+					LeftCenterTransform.AddToTranslation(PoseTransform.GetLocation());
+					LeftCenterTransform.ConcatenateRotation(PoseTransform.Rotator().Quaternion());
+					bLeftTransformValid = true;
+				}
+
+				bool bRightTransformValid = false;
 				FTransform RightCenterTransform;
-				AppFramework.GetTransform(EyeTrackingStaticData.right_center, RightCenterTransform, FailReason);
-				RightCenterTransform.AddToTranslation(PoseTransform.GetLocation());
-				RightCenterTransform.ConcatenateRotation(PoseTransform.Rotator().Quaternion());
-
-				//average the left and right eye
-				UnfilteredEyeTrackingData.AverageGazeOrigin = (LeftCenterTransform.GetLocation() + RightCenterTransform.GetLocation()) * .5f;
-				//get focal point
-				UnfilteredEyeTrackingData.WorldAverageGazeConvergencePoint = FixationTransform.GetTranslation();
-				//get the gaze vector (Point-Eye)
-				UnfilteredEyeTrackingData.AverageGazeRay = (UnfilteredEyeTrackingData.WorldAverageGazeConvergencePoint - UnfilteredEyeTrackingData.AverageGazeOrigin);
-				UnfilteredEyeTrackingData.AverageGazeRay.Normalize();
-
-				UnfilteredEyeTrackingData.LeftOriginPoint = LeftCenterTransform.GetLocation();
-				UnfilteredEyeTrackingData.RightOriginPoint = RightCenterTransform.GetLocation();
+				if (AppFramework.GetTransform(EyeTrackingStaticData.right_center, RightCenterTransform, FailReason))
+				{
+					RightCenterTransform.AddToTranslation(PoseTransform.GetLocation());
+					RightCenterTransform.ConcatenateRotation(PoseTransform.Rotator().Quaternion());
+					bRightTransformValid = true;
+				}
+
+				if (bLeftTransformValid && bRightTransformValid)
+				{
+					//average the left and right eye
+					UnfilteredEyeTrackingData.AverageGazeOrigin = (LeftCenterTransform.GetLocation() + RightCenterTransform.GetLocation()) * .5f;
+					//get the gaze vector (Point-Eye)
+					UnfilteredEyeTrackingData.AverageGazeRay = (UnfilteredEyeTrackingData.WorldAverageGazeConvergencePoint - UnfilteredEyeTrackingData.AverageGazeOrigin);
+					UnfilteredEyeTrackingData.AverageGazeRay.Normalize();
+
+					UnfilteredEyeTrackingData.LeftOriginPoint = LeftCenterTransform.GetLocation();
+					UnfilteredEyeTrackingData.RightOriginPoint = RightCenterTransform.GetLocation();
+				}
 
 				UnfilteredEyeTrackingData.Confidence = TempTrackingState.fixation_confidence;
 
 				UnfilteredEyeTrackingData.bLeftBlink = TempTrackingState.left_blink;
 				UnfilteredEyeTrackingData.bRightBlink = TempTrackingState.right_blink;
 
+				UnfilteredEyeTrackingData.FixationDepthIsUncomfortable = TempTrackingState.fixation_depth_is_uncomfortable;
+				UnfilteredEyeTrackingData.FixationDepthViolationHasOccurred = TempTrackingState.fixation_depth_violation_has_occurred;
+				UnfilteredEyeTrackingData.RemainingTimeAtUncomfortableDepth = TempTrackingState.remaining_time_at_uncomfortable_depth;
+
 				//UE_LOG(LogCore, Warning, TEXT("   SUCCESS -> "));
 				//UE_LOG(LogCore, Warning, TEXT("        Gaze Origin %f, %f, %f"), UnfilteredEyeTrackingData.AverageGazeOrigin.X, UnfilteredEyeTrackingData.AverageGazeOrigin.Y, UnfilteredEyeTrackingData.AverageGazeOrigin.Z);
 				//UE_LOG(LogCore, Warning, TEXT("        Converge    %f, %f, %f"), UnfilteredEyeTrackingData.WorldAverageGazeConvergencePoint.X, UnfilteredEyeTrackingData.WorldAverageGazeConvergencePoint.Y, UnfilteredEyeTrackingData.WorldAverageGazeConvergencePoint.Z);
@@ -144,26 +168,29 @@ bool FMagicLeapVREyeTracker::Tick(float DeltaTime)
 	}
 	else
 	{
-		if (!bInitialized && GEngine->XRSystem.IsValid() && GEngine->XRSystem->GetHMDDevice() && static_cast<FMagicLeapHMD*>(GEngine->XRSystem->GetHMDDevice())->IsPerceptionEnabled())
+		if (!bInitialized && IMagicLeapPlugin::Get().IsMagicLeapHMDValid())
 		{
 #if !PLATFORM_MAC
-			//keep trying until we are successful in creating one
-			ML_FUNCTION_WRAPPER(EyeTrackingHandle = MLEyeTrackingCreate());
-			bInitialized = MLHandleIsValid(EyeTrackingHandle);
-			if (bInitialized)
+			if (static_cast<FMagicLeapHMD*>(GEngine->XRSystem->GetHMDDevice())->IsPerceptionEnabled())
 			{
-				UE_LOG(LogCore, Warning, TEXT("   VR Eye Tracker Created"));
-				// Needs to be called only once.
-				MLEyeTrackingGetStaticData(EyeTrackingHandle, &EyeTrackingStaticData);
+				//keep trying until we are successful in creating one
+				MLResult CreateResult = MLResult_UnspecifiedFailure;
+				ML_FUNCTION_WRAPPER(CreateResult = MLEyeTrackingCreate(&EyeTrackingHandle));
+				bInitialized = CreateResult == MLResult_Ok && MLHandleIsValid(EyeTrackingHandle);
+				if (bInitialized)
+				{
+					UE_LOG(LogCore, Warning, TEXT("   VR Eye Tracker Created"));
+					// Needs to be called only once.
+					if (MLEyeTrackingGetStaticData(EyeTrackingHandle, &EyeTrackingStaticData) != MLResult_Ok)
+					{
+						UE_LOG(LogCore, Warning, TEXT("   Unable to get Eye Tracker Static Data"));
+					}
+				}
 			}
-			else
-			{
-				UE_LOG(LogCore, Warning, TEXT("   Unable to create VR Eye Tracker"));
-			}
-#endif
+#endif //!PLATFORM_MAC
 		}
 	}
-#endif //WITH_MLSDK && PLATFORM_LUMIN
+#endif //WITH_MLSDK
 	return true;
 }
 
diff --git a/Engine/Plugins/Lumin/MagicLeap/Source/MagicLeapEyeTracker/Public/MagicLeapEyeTrackerModule.h b/Engine/Plugins/Lumin/MagicLeap/Source/MagicLeapEyeTracker/Public/MagicLeapEyeTrackerModule.h
--- a/Engine/Plugins/Lumin/MagicLeap/Source/MagicLeapEyeTracker/Public/MagicLeapEyeTrackerModule.h
+++ b/Engine/Plugins/Lumin/MagicLeap/Source/MagicLeapEyeTracker/Public/MagicLeapEyeTrackerModule.h
@@ -28,6 +28,31 @@ public:
 	bool RightEyeBlinked;
 };
 
+USTRUCT(BlueprintType)
+struct FMagicLeapFixationComfort
+{
+	GENERATED_BODY()
+
+public:
+	/**
+	  Is the user's fixation point too close for sustained use. This value is true if the user is
+      focused on a point that is within 37 cm of the eyeball centers. 
+	*/
+	UPROPERTY(EditAnywhere, BlueprintReadWrite, Category = "Eye Fixation Comfort")
+	bool FixationDepthIsUncomfortable;
+
+	/** True if the user has fixated on a point closer than 37 cm for longer than 10 seconds within the last minute. */
+	UPROPERTY(EditAnywhere, BlueprintReadWrite, Category = "Eye Fixation Comfort")
+	bool FixationDepthViolationHasOccurred;
+
+	/**
+	  Number of seconds remaining that the user may be fixated at an uncomfortable depth. If this
+	  persists for too long, the system may take action to move the fixation point further away.
+	*/
+	UPROPERTY(EditAnywhere, BlueprintReadWrite, Category = "Eye Fixation Comfort")
+	float RemainingTimeAtUncomfortableDepth;
+};
+
 /**
 * The public interface of the Magic Leap Eye Tracking Module.
 */
@@ -82,6 +107,7 @@ public:
 
 	bool IsEyeTrackerCalibrated() const;
 	bool GetEyeBlinkState(FMagicLeapEyeBlinkState& BlinkState) const;
+	bool GetFixationComfort(FMagicLeapFixationComfort& FixationComfort) const;
 
 private:
 	friend class FMagicLeapEyeTrackerModule;
@@ -133,4 +159,7 @@ public:
 
 	UFUNCTION(BlueprintCallable, BlueprintPure, Category = "Eye Tracking|MagicLeap")
 	static bool GetEyeBlinkState(FMagicLeapEyeBlinkState& BlinkState);
+
+	UFUNCTION(BlueprintCallable, BlueprintPure, Category = "Eye Tracking|MagicLeap")
+	static bool GetFixationComfort(FMagicLeapFixationComfort& FixationComfort);
 };
diff --git a/Engine/Plugins/Lumin/MagicLeap/Source/MagicLeapEyeTracker/Public/MagicLeapEyeTrackerTypes.h b/Engine/Plugins/Lumin/MagicLeap/Source/MagicLeapEyeTracker/Public/MagicLeapEyeTrackerTypes.h
--- a/Engine/Plugins/Lumin/MagicLeap/Source/MagicLeapEyeTracker/Public/MagicLeapEyeTrackerTypes.h
+++ b/Engine/Plugins/Lumin/MagicLeap/Source/MagicLeapEyeTracker/Public/MagicLeapEyeTrackerTypes.h
@@ -33,6 +33,13 @@ public:
 	bool bLeftBlink;
 	UPROPERTY(VisibleAnywhere, BlueprintReadOnly, Category = "Blink Data")
 	bool bRightBlink;
+
+	UPROPERTY(EditAnywhere, BlueprintReadWrite, Category = "Eye Fixation Comfort")
+	bool FixationDepthIsUncomfortable;
+	UPROPERTY(EditAnywhere, BlueprintReadWrite, Category = "Eye Fixation Comfort")
+	bool FixationDepthViolationHasOccurred;
+	UPROPERTY(EditAnywhere, BlueprintReadWrite, Category = "Eye Fixation Comfort")
+	float RemainingTimeAtUncomfortableDepth;
 };
 
 struct FMagicLeapVRStabilityData
diff --git a/Engine/Plugins/Lumin/MagicLeap/Source/MagicLeapGestures/Private/MagicLeapGestures.cpp b/Engine/Plugins/Lumin/MagicLeap/Source/MagicLeapGestures/Private/MagicLeapGestures.cpp
--- a/Engine/Plugins/Lumin/MagicLeap/Source/MagicLeapGestures/Private/MagicLeapGestures.cpp
+++ b/Engine/Plugins/Lumin/MagicLeap/Source/MagicLeapGestures/Private/MagicLeapGestures.cpp
@@ -1,6 +1,7 @@
 // Copyright 1998-2018 Epic Games, Inc. All Rights Reserved.
 
 #include "MagicLeapGestures.h"
+#include "IMagicLeapPlugin.h"
 #include "MagicLeapHMD.h"
 #include "MagicLeapMath.h"
 #include "Framework/Application/SlateApplication.h"
@@ -405,7 +406,7 @@ void FMagicLeapGestures::Disable()
 #if WITH_MLSDK
 	if (MLHandleIsValid(GestureTracker))
 	{
-		if (MLGestureTrackingDestroy(GestureTracker))
+		if (MLResult_Ok == MLGestureTrackingDestroy(GestureTracker))
 		{
 			GestureTracker = ML_INVALID_HANDLE;
 		}
@@ -464,7 +465,7 @@ bool FMagicLeapGestures::SetConfiguration(const TArray<EStaticGestures>& StaticG
 	config.pose_filter_level = UnrealToMLGestureFilterLevel(GestureFilterLevel);
 	config.handtype_filter_level = UnrealToMLGestureFilterLevel(HandSwitchingFilterLevel);
 
-	return MLGestureSetConfiguration(GestureTracker, &config);
+	return MLGestureSetConfiguration(GestureTracker, &config) == MLResult_Ok;
 #else
 	return false;
 #endif //WITH_MLSDK
@@ -481,7 +482,7 @@ bool FMagicLeapGestures::GetConfiguration(TArray<EStaticGestures>& ActiveStaticG
 	}
 
 	MLGestureConfiguration config;
-	bool validResult = MLGestureGetConfiguration(GestureTracker, &config);
+	bool validResult = MLGestureGetConfiguration(GestureTracker, &config) == MLResult_Ok;
 
 	if (validResult)
 	{
@@ -529,14 +530,14 @@ void FMagicLeapGestures::UpdateTrackerData()
 	if (MLHandleIsValid(GestureTracker))
 	{
 		OldData = GestureData;
-		bIsGestureStateValid = MLGestureGetData(GestureTracker, &GestureData);
+		bIsGestureStateValid = MLGestureGetData(GestureTracker, &GestureData) == MLResult_Ok;
 	}
 	else
 	{
 		bIsGestureStateValid = false;
 	}
 
-	if (bIsGestureStateValid)
+	if (bIsGestureStateValid && IMagicLeapPlugin::Get().IsMagicLeapHMDValid())
 	{
 		const FAppFramework& AppFramework = static_cast<FMagicLeapHMD*>(GEngine->XRSystem->GetHMDDevice())->GetAppFrameworkConst();
 		check(AppFramework.IsInitialized());
@@ -717,11 +718,13 @@ void FMagicLeapGestures::ConditionallyEnable()
 #if WITH_MLSDK
 	if (!MLHandleIsValid(GestureTracker) && GEngine->XRSystem.IsValid() && GEngine->XRSystem->GetHMDDevice())
 	{
-		GestureTracker = MLGestureTrackingCreate();
-
-		if (MLHandleIsValid(GestureTracker))
+		MLResult CreateResult = MLGestureTrackingCreate(&GestureTracker);
+		if (CreateResult == MLResult_Ok && MLHandleIsValid(GestureTracker))
 		{
-			MLGestureGetStaticData(GestureTracker, &StaticData);
+			if (MLResult_Ok != MLGestureGetStaticData(GestureTracker, &StaticData))
+			{
+				UE_LOG(LogMagicLeapGestures, Error, TEXT("Error getting gesture tracker static data."));
+			}
 		}
 		else
 		{
@@ -742,7 +745,7 @@ void FMagicLeapGestures::OnAppPause()
 	{
 		MLGestureConfiguration GestureConfig;
 
-		if (!MLGestureGetConfiguration(GestureTracker, &GestureConfig))
+		if (MLResult_Ok != MLGestureGetConfiguration(GestureTracker, &GestureConfig))
 		{
 			UE_LOG(LogMagicLeapGestures, Error, TEXT("Failed to retrieve gesture tracking configuration on application pause."));
 		}
@@ -758,7 +761,7 @@ void FMagicLeapGestures::OnAppPause()
 			{
 				GestureConfig.gesture_pipeline_enabled = false;
 
-				if (!MLGestureSetConfiguration(GestureTracker, &GestureConfig))
+				if (MLResult_Ok != MLGestureSetConfiguration(GestureTracker, &GestureConfig))
 				{
 					UE_LOG(LogMagicLeapGestures, Error, TEXT("Failed to disable gesture tracking on application pause."));
 				}
@@ -789,7 +792,7 @@ void FMagicLeapGestures::OnAppResume()
 		}
 		else
 		{
-			if (!MLGestureGetConfiguration(GestureTracker, &GestureConfig))
+			if (MLResult_Ok != MLGestureGetConfiguration(GestureTracker, &GestureConfig))
 			{
 				UE_LOG(LogMagicLeapGestures, Error, TEXT("Failed to retrieve gesture tracking configuration on application resume."));
 			}
@@ -797,7 +800,7 @@ void FMagicLeapGestures::OnAppResume()
 			{
 				GestureConfig.gesture_pipeline_enabled = true;
 
-				if (!MLGestureSetConfiguration(GestureTracker, &GestureConfig))
+				if (MLResult_Ok != MLGestureSetConfiguration(GestureTracker, &GestureConfig))
 				{
 					UE_LOG(LogMagicLeapGestures, Error, TEXT("Failed to re-enable gesture tracking on application resume."));
 				}
diff --git a/Engine/Plugins/Lumin/MagicLeap/Source/MagicLeapIdentity/Private/MagicLeapIdentity.cpp b/Engine/Plugins/Lumin/MagicLeap/Source/MagicLeapIdentity/Private/MagicLeapIdentity.cpp
--- a/Engine/Plugins/Lumin/MagicLeap/Source/MagicLeapIdentity/Private/MagicLeapIdentity.cpp
+++ b/Engine/Plugins/Lumin/MagicLeap/Source/MagicLeapIdentity/Private/MagicLeapIdentity.cpp
@@ -22,7 +22,7 @@ public:
 		APISetup.Startup();
 #if WITH_MLSDK
 		APISetup.LoadDLL(TEXT("ml_identity"));
-#endif
+#endif //WITH_MLSDK
 	}
 
 	void ShutdownModule() override
@@ -73,109 +73,73 @@ public:
 		UMagicLeapIdentity::FModifyIdentityAttributeValueDelegate RequestDelegate;
 	};
 
-	EMagicLeapIdentityError MLToUnrealIdentityError(MLIdentityError error)
+	EMagicLeapIdentityError MLToUnrealIdentityError(MLResult error)
 	{
+		#define MLRESULTCASE(x) case MLResult_##x: { return EMagicLeapIdentityError::x; }
+		#define MLIDENTITYRESULTCASE(x) case MLIdentityResult_##x: { return EMagicLeapIdentityError::x; }
 		switch (error)
 		{
-		case MLIdentityError_Ok:
-			return EMagicLeapIdentityError::Ok;
-		case MLIdentityError_FailedToConnectToLocalService:
-			return EMagicLeapIdentityError::FailedToConnectToLocalService;
-		case MLIdentityError_FailedToConnectToCloudService:
-			return EMagicLeapIdentityError::FailedToConnectToCloudService;
-		case MLIdentityError_CloudAuthentication:
-			return EMagicLeapIdentityError::CloudAuthentication;
-		case MLIdentityError_InvalidInformationFromCloud:
-			return EMagicLeapIdentityError::InvalidInformationFromCloud;
-		case MLIdentityError_InvalidArgument:
-			return EMagicLeapIdentityError::InvalidArgument;
-		case MLIdentityError_AsyncOperationNotComplete:
-			return EMagicLeapIdentityError::AsyncOperationNotComplete;
-		case MLIdentityError_OtherError:
-			return EMagicLeapIdentityError::OtherError;
+			MLRESULTCASE(Ok)
+			MLRESULTCASE(InvalidParam)
+			MLRESULTCASE(AllocFailed)
+			MLRESULTCASE(PrivilegeDenied)
+			MLRESULTCASE(UnspecifiedFailure)
+			MLIDENTITYRESULTCASE(FailedToConnectToLocalService)
+			MLIDENTITYRESULTCASE(FailedToConnectToCloudService)
+			MLIDENTITYRESULTCASE(CloudAuthentication)
+			MLIDENTITYRESULTCASE(InvalidInformationFromCloud)
+			MLIDENTITYRESULTCASE(NotLoggedIn)
+			MLIDENTITYRESULTCASE(ExpiredCredentials)
+			MLIDENTITYRESULTCASE(FailedToGetUserProfile)
+			MLIDENTITYRESULTCASE(Unauthorized)
+			MLIDENTITYRESULTCASE(CertificateError)
+			MLIDENTITYRESULTCASE(RejectedByCloud)
+			MLIDENTITYRESULTCASE(AlreadyLoggedIn)
+			MLIDENTITYRESULTCASE(ModifyIsNotSupported)
+			MLIDENTITYRESULTCASE(NetworkError)
+			default:
+				break;
 		}
 
-		return EMagicLeapIdentityError::OtherError;
+		return EMagicLeapIdentityError::UnspecifiedFailure;
 	}
 
-	EMagicLeapIdentityAttribute MLToUnrealIdentityAttribute(MLIdentityAttributeEnum attribute)
+	EMagicLeapIdentityKey MLToUnrealIdentityAttribute(MLIdentityAttributeKey attribute)
 	{
+#define MLIDENTITYKEYCASE(x) case MLIdentityAttributeKey_##x: { return EMagicLeapIdentityKey::x; }
 		switch (attribute)
 		{
-		case MLIdentityAttributeEnum_UserId:
-			return EMagicLeapIdentityAttribute::UserID;
-		case MLIdentityAttributeEnum_GivenName:
-			return EMagicLeapIdentityAttribute::GivenName;
-		case MLIdentityAttributeEnum_FamilyName:
-			return EMagicLeapIdentityAttribute::FamilyName;
-		case MLIdentityAttributeEnum_Email:
-			return EMagicLeapIdentityAttribute::Email;
-		case MLIdentityAttributeEnum_Status:
-			return EMagicLeapIdentityAttribute::Status;
-		case MLIdentityAttributeEnum_TermsAccepted:
-			return EMagicLeapIdentityAttribute::TermsAccepted;
-		case MLIdentityAttributeEnum_Birthday:
-			return EMagicLeapIdentityAttribute::Birthday;
-		case MLIdentityAttributeEnum_Company:
-			return EMagicLeapIdentityAttribute::Company;
-		case MLIdentityAttributeEnum_Industry:
-			return EMagicLeapIdentityAttribute::Industry;
-		case MLIdentityAttributeEnum_Location:
-			return EMagicLeapIdentityAttribute::Location;
-		case MLIdentityAttributeEnum_Tagline:
-			return EMagicLeapIdentityAttribute::Tagline;
-		case MLIdentityAttributeEnum_PhoneNumber:
-			return EMagicLeapIdentityAttribute::PhoneNumber;
-		case MLIdentityAttributeEnum_Avatar2D:
-			return EMagicLeapIdentityAttribute::Avatar2D;
-		case MLIdentityAttributeEnum_Avatar3D:
-			return EMagicLeapIdentityAttribute::Avatar3D;
-		case MLIdentityAttributeEnum_IsDeveloper:
-			return EMagicLeapIdentityAttribute::IsDeveloper;
-		case MLIdentityAttributeEnum_Unknown:
-			return EMagicLeapIdentityAttribute::Unknown;
+			MLIDENTITYKEYCASE(GivenName)
+			MLIDENTITYKEYCASE(FamilyName)
+			MLIDENTITYKEYCASE(Email)
+			MLIDENTITYKEYCASE(Bio)
+			MLIDENTITYKEYCASE(PhoneNumber)
+			MLIDENTITYKEYCASE(Avatar2D)
+			MLIDENTITYKEYCASE(Avatar3D)
+			default:
+				break;
 		}
-		return EMagicLeapIdentityAttribute::Unknown;
+		return EMagicLeapIdentityKey::Unknown;
+#undef MLIDENTITYKEYCASE
 	}
 
-	MLIdentityAttributeEnum UnrealToMLIdentityAttribute(EMagicLeapIdentityAttribute attribute)
+	MLIdentityAttributeKey UnrealToMLIdentityAttribute(EMagicLeapIdentityKey attribute)
 	{
+#define MLIDENTITYKEYCASE(x) case EMagicLeapIdentityKey::x: { return MLIdentityAttributeKey_##x; }
 		switch (attribute)
 		{
-		case EMagicLeapIdentityAttribute::UserID:
-			return MLIdentityAttributeEnum_UserId;
-		case EMagicLeapIdentityAttribute::GivenName:
-			return MLIdentityAttributeEnum_GivenName;
-		case EMagicLeapIdentityAttribute::FamilyName:
-			return MLIdentityAttributeEnum_FamilyName;
-		case EMagicLeapIdentityAttribute::Email:
-			return MLIdentityAttributeEnum_Email;
-		case EMagicLeapIdentityAttribute::Status:
-			return MLIdentityAttributeEnum_Status;
-		case EMagicLeapIdentityAttribute::TermsAccepted:
-			return MLIdentityAttributeEnum_TermsAccepted;
-		case EMagicLeapIdentityAttribute::Birthday:
-			return MLIdentityAttributeEnum_Birthday;
-		case EMagicLeapIdentityAttribute::Company:
-			return MLIdentityAttributeEnum_Company;
-		case EMagicLeapIdentityAttribute::Industry:
-			return MLIdentityAttributeEnum_Industry;
-		case EMagicLeapIdentityAttribute::Location:
-			return MLIdentityAttributeEnum_Location;
-		case EMagicLeapIdentityAttribute::Tagline:
-			return MLIdentityAttributeEnum_Tagline;
-		case EMagicLeapIdentityAttribute::PhoneNumber:
-			return MLIdentityAttributeEnum_PhoneNumber;
-		case EMagicLeapIdentityAttribute::Avatar2D:
-			return MLIdentityAttributeEnum_Avatar2D;
-		case EMagicLeapIdentityAttribute::Avatar3D:
-			return MLIdentityAttributeEnum_Avatar3D;
-		case EMagicLeapIdentityAttribute::IsDeveloper:
-			return MLIdentityAttributeEnum_IsDeveloper;
-		case EMagicLeapIdentityAttribute::Unknown:
-			return MLIdentityAttributeEnum_Unknown;
+			MLIDENTITYKEYCASE(GivenName)
+			MLIDENTITYKEYCASE(FamilyName)
+			MLIDENTITYKEYCASE(Email)
+			MLIDENTITYKEYCASE(Bio)
+			MLIDENTITYKEYCASE(PhoneNumber)
+			MLIDENTITYKEYCASE(Avatar2D)
+			MLIDENTITYKEYCASE(Avatar3D)
+			default:
+				break;
 		}
-		return MLIdentityAttributeEnum_Unknown;
+		return MLIdentityAttributeKey_Unknown;
+#undef MLIDENTITYKEYCASE
 	}
 
 
@@ -193,66 +157,74 @@ UMagicLeapIdentity::~UMagicLeapIdentity()
 	delete Impl;
 }
 
-EMagicLeapIdentityError UMagicLeapIdentity::GetAllAvailableAttributes(TArray<EMagicLeapIdentityAttribute>& AvailableAttributes)
+EMagicLeapIdentityError UMagicLeapIdentity::GetAllAvailableAttributes(TArray<EMagicLeapIdentityKey>& AvailableAttributes)
 {
 #if WITH_MLSDK
 	MLIdentityProfile* profile = nullptr;
-	MLIdentityError result = MLIdentityGetAttributeNames(&profile);
+	MLResult Result = MLIdentityGetAttributeNames(&profile);
 
 	AvailableAttributes.Empty();
-	if (result == MLIdentityError_Ok && profile != nullptr)
+	if (Result == MLResult_Ok && profile != nullptr)
 	{
 		for (uint32 i = 0; i < static_cast<uint32>(profile->attribute_count); ++i)
 		{
 			const MLIdentityAttribute* attribute = profile->attribute_ptrs[i];
-			AvailableAttributes.Add(Impl->MLToUnrealIdentityAttribute(attribute->enum_value));
+			AvailableAttributes.Add(Impl->MLToUnrealIdentityAttribute(attribute->key));
 		}
 
 		MLIdentityReleaseUserProfile(profile);
 	}
 
-	return Impl->MLToUnrealIdentityError(result);
+	return Impl->MLToUnrealIdentityError(Result);
 #else
-	return EMagicLeapIdentityError::OtherError;
+	return EMagicLeapIdentityError::UnspecifiedFailure;
 #endif //WITH_MLSDK
-
 }
 
 void UMagicLeapIdentity::GetAllAvailableAttributesAsync(const FAvailableIdentityAttributesDelegate& ResultDelegate)
 {
 #if WITH_MLSDK
-	Impl->AllAvailableAttribsFutures.Add(MLIdentityGetAttributeNamesAsync(), ResultDelegate);
+	MLInvokeFuture* InvokeFuture = nullptr;
+	MLResult Result = MLIdentityGetAttributeNamesAsync(&InvokeFuture);
+	if (Result == MLResult_Ok)
+	{
+		Impl->AllAvailableAttribsFutures.Add(InvokeFuture, ResultDelegate);
+	}
+	else
+	{
+		UE_LOG(LogMagicLeapIdentity, Error, TEXT("MLIdentityGetAttributeNamesAsync failed with error %d!"), Result);
+	}
 #endif //WITH_MLSDK
 }
 
-EMagicLeapIdentityError UMagicLeapIdentity::RequestAttributeValue(const TArray<EMagicLeapIdentityAttribute>& Attribute, TArray<FMagicLeapIdentityAttribute>& AttributeValue)
+EMagicLeapIdentityError UMagicLeapIdentity::RequestAttributeValue(const TArray<EMagicLeapIdentityKey>& Attribute, TArray<FMagicLeapIdentityAttribute>& AttributeValue)
 {
 #if WITH_MLSDK
 	MLIdentityProfile* profile = nullptr;
-	TArray<MLIdentityAttributeEnum> mlAttributes;
-	for (EMagicLeapIdentityAttribute attrib : Attribute)
+	TArray<MLIdentityAttributeKey> mlAttributes;
+	for (EMagicLeapIdentityKey attrib : Attribute)
 	{
 		mlAttributes.Add(Impl->UnrealToMLIdentityAttribute(attrib));
 	}
-	MLIdentityError result = MLIdentityGetKnownAttributeNames(mlAttributes.GetData(), mlAttributes.Num(), &profile);
+	MLResult Result = MLIdentityGetKnownAttributeNames(mlAttributes.GetData(), mlAttributes.Num(), &profile);
 
 	AttributeValue.Empty();
-	if (result == MLIdentityError_Ok && profile != nullptr)
+	if (Result == MLResult_Ok && profile != nullptr)
 	{
 		for (uint32 i = 0; i < static_cast<uint32>(profile->attribute_count); ++i)
 		{
 			profile->attribute_ptrs[i]->is_requested = true;
 		}
 
-		result = MLIdentityRequestAttributeValues(profile);
+		Result = MLIdentityRequestAttributeValues(profile);
 
-		if (result == MLIdentityError_Ok)
+		if (Result == MLResult_Ok)
 		{
 			for (uint32 i = 0; i < static_cast<uint32>(profile->attribute_count); ++i)
 			{
 				if (profile->attribute_ptrs[i]->is_granted)
 				{
-					AttributeValue.Add(FMagicLeapIdentityAttribute(Impl->MLToUnrealIdentityAttribute(profile->attribute_ptrs[i]->enum_value), FString(ANSI_TO_TCHAR(profile->attribute_ptrs[i]->value))));
+					AttributeValue.Add(FMagicLeapIdentityAttribute(Impl->MLToUnrealIdentityAttribute(profile->attribute_ptrs[i]->key), FString(ANSI_TO_TCHAR(profile->attribute_ptrs[i]->value))));
 				}
 			}
 		}
@@ -260,49 +232,48 @@ EMagicLeapIdentityError UMagicLeapIdentity::RequestAttributeValue(const TArray<E
 		MLIdentityReleaseUserProfile(profile);
 	}
 
-	return Impl->MLToUnrealIdentityError(result);
+	return Impl->MLToUnrealIdentityError(Result);
 #else
-	return EMagicLeapIdentityError::OtherError;
+	return EMagicLeapIdentityError::UnspecifiedFailure;
 #endif //WITH_MLSDK
 }
 
-EMagicLeapIdentityError UMagicLeapIdentity::RequestAttributeValueAsync(const TArray<EMagicLeapIdentityAttribute>& Attribute, const FRequestIdentityAttributeValueDelegate& ResultDelegate)
+EMagicLeapIdentityError UMagicLeapIdentity::RequestAttributeValueAsync(const TArray<EMagicLeapIdentityKey>& Attribute, const FRequestIdentityAttributeValueDelegate& ResultDelegate)
 {
 #if WITH_MLSDK
 	MLIdentityProfile* profile = nullptr;
-	TArray<MLIdentityAttributeEnum> mlAttributes;
-	for (EMagicLeapIdentityAttribute attrib : Attribute)
+	TArray<MLIdentityAttributeKey> mlAttributes;
+	for (EMagicLeapIdentityKey attrib : Attribute)
 	{
 		mlAttributes.Add(Impl->UnrealToMLIdentityAttribute(attrib));
 	}
-	MLIdentityError result = MLIdentityGetKnownAttributeNames(mlAttributes.GetData(), mlAttributes.Num(), &profile);
+	MLResult Result = MLIdentityGetKnownAttributeNames(mlAttributes.GetData(), mlAttributes.Num(), &profile);
 
-	if (result == MLIdentityError_Ok && profile != nullptr)
+	if (Result == MLResult_Ok && profile != nullptr)
 	{
 		for (uint32 i = 0; i < static_cast<uint32>(profile->attribute_count); ++i)
 		{
 			profile->attribute_ptrs[i]->is_requested = true;
 		}
 
-		Impl->AllRequestAttribsFutures.Add(MLIdentityRequestAttributeValuesAsync(profile), FIdentityImpl::FRequestAttribData(profile, ResultDelegate));
+		MLInvokeFuture* InvokeFuture = nullptr;
+		MLResult RequestAtrrValRes = MLIdentityRequestAttributeValuesAsync(profile, &InvokeFuture);
+		if (RequestAtrrValRes == MLResult_Ok)
+		{
+			Impl->AllRequestAttribsFutures.Add(InvokeFuture, FIdentityImpl::FRequestAttribData(profile, ResultDelegate));
+		}
+		else
+		{
+			UE_LOG(LogMagicLeapIdentity, Error, TEXT("MLIdentityRequestAttributeValuesAsync failed with error %d!"), RequestAtrrValRes);
+		}
 	}
 
-	return Impl->MLToUnrealIdentityError(result);
+	return Impl->MLToUnrealIdentityError(Result);
 #else
-	return EMagicLeapIdentityError::OtherError;
+	return EMagicLeapIdentityError::UnspecifiedFailure;
 #endif //WITH_MLSDK
 }
 
-EMagicLeapIdentityError UMagicLeapIdentity::ModifyAttributeValue(const TArray<FMagicLeapIdentityAttribute>& UpdatedAttributeValue, TArray<EMagicLeapIdentityAttribute>& AttributesUpdatedSuccessfully)
-{
-	return EMagicLeapIdentityError::OtherError;
-}
-
-EMagicLeapIdentityError UMagicLeapIdentity::ModifyAttributeValueAsync(const TArray<FMagicLeapIdentityAttribute>& UpdatedAttributeValue, const FModifyIdentityAttributeValueDelegate& ResultDelegate)
-{
-	return EMagicLeapIdentityError::OtherError;
-}
-
 UWorld* UMagicLeapIdentity::GetWorld() const
 {
 	return Cast<UWorld>(GetOuter());
@@ -317,23 +288,28 @@ void UMagicLeapIdentity::Tick(float DeltaTime)
 	for (const auto& AvailableAttribsFuture : Impl->AllAvailableAttribsFutures)
 	{
 		MLInvokeFuture* future = AvailableAttribsFuture.Key;
-		MLIdentityError error = MLIdentityError_OtherError;
 		MLIdentityProfile* profile = nullptr;
+		MLResult Result = MLIdentityGetAttributeNamesWait(future, 0, &profile);
 
-		if (MLIdentityGetAttributeNamesWait(future, 0, &error, &profile))
+		TArray<EMagicLeapIdentityKey> AvailableAttributes;
+		if (Result == MLResult_Ok)
 		{
 			FuturesToDelete.Add(future);
-			TArray<EMagicLeapIdentityAttribute> AvailableAttributes;
-			if (error == MLIdentityError_Ok && profile != nullptr)
+			if (profile != nullptr)
 			{
 				for (uint32 i = 0; i < static_cast<uint32>(profile->attribute_count); ++i)
 				{
 					const MLIdentityAttribute* attribute = profile->attribute_ptrs[i];
-					AvailableAttributes.Add(Impl->MLToUnrealIdentityAttribute(attribute->enum_value));
+					AvailableAttributes.Add(Impl->MLToUnrealIdentityAttribute(attribute->key));
 				}
 				MLIdentityReleaseUserProfile(profile);
 			}
-			AvailableAttribsFuture.Value.ExecuteIfBound(Impl->MLToUnrealIdentityError(error), AvailableAttributes);
+			AvailableAttribsFuture.Value.ExecuteIfBound(Impl->MLToUnrealIdentityError(Result), AvailableAttributes);
+		}
+		else if (Result != MLResult_Pending)
+		{
+			FuturesToDelete.Add(future);			
+			AvailableAttribsFuture.Value.ExecuteIfBound(Impl->MLToUnrealIdentityError(Result), AvailableAttributes);
 		}
 	}
 
@@ -348,27 +324,30 @@ void UMagicLeapIdentity::Tick(float DeltaTime)
 	for (const auto& RequestAttribsFuture : Impl->AllRequestAttribsFutures)
 	{
 		MLInvokeFuture* future = RequestAttribsFuture.Key;
-		MLIdentityError error = MLIdentityError_OtherError;
+		MLResult Result = MLIdentityRequestAttributeValuesWait(future, 0);
 
-		if (MLIdentityRequestAttributeValuesWait(future, 0, &error))
+		TArray<FMagicLeapIdentityAttribute> AttributeValue;
+		if (Result == MLResult_Ok)
 		{
 			FuturesToDelete.Add(future);
-			TArray<FMagicLeapIdentityAttribute> AttributeValue;
 			MLIdentityProfile* profile = RequestAttribsFuture.Value.Profile;
 
-			if (error == MLIdentityError_Ok)
+			for (uint32 i = 0; i < static_cast<uint32>(profile->attribute_count); ++i)
 			{
-				for (uint32 i = 0; i < static_cast<uint32>(profile->attribute_count); ++i)
+				if (profile->attribute_ptrs[i]->is_granted)
 				{
-					if (profile->attribute_ptrs[i]->is_granted)
-					{
-						AttributeValue.Add(FMagicLeapIdentityAttribute(Impl->MLToUnrealIdentityAttribute(profile->attribute_ptrs[i]->enum_value), FString(ANSI_TO_TCHAR(profile->attribute_ptrs[i]->value))));
-					}
+					AttributeValue.Add(FMagicLeapIdentityAttribute(Impl->MLToUnrealIdentityAttribute(profile->attribute_ptrs[i]->key), FString(ANSI_TO_TCHAR(profile->attribute_ptrs[i]->value))));
 				}
 			}
-			RequestAttribsFuture.Value.RequestDelegate.ExecuteIfBound(Impl->MLToUnrealIdentityError(error), AttributeValue);
+
+			RequestAttribsFuture.Value.RequestDelegate.ExecuteIfBound(Impl->MLToUnrealIdentityError(Result), AttributeValue);
 			MLIdentityReleaseUserProfile(profile);
 		}
+		else if (Result != MLResult_Pending)
+		{
+			FuturesToDelete.Add(future);			
+			RequestAttribsFuture.Value.RequestDelegate.ExecuteIfBound(Impl->MLToUnrealIdentityError(Result), AttributeValue);
+		}
 	}
 
 	for (const MLInvokeFuture* future : FuturesToDelete)
diff --git a/Engine/Plugins/Lumin/MagicLeap/Source/MagicLeapIdentity/Public/MagicLeapIdentity.h b/Engine/Plugins/Lumin/MagicLeap/Source/MagicLeapIdentity/Public/MagicLeapIdentity.h
--- a/Engine/Plugins/Lumin/MagicLeap/Source/MagicLeapIdentity/Public/MagicLeapIdentity.h
+++ b/Engine/Plugins/Lumin/MagicLeap/Source/MagicLeapIdentity/Public/MagicLeapIdentity.h
@@ -25,7 +25,7 @@ public:
 	  @param ResultCode Error code when getting the available attributes.
 	  @param AvailableAttributes List of attributes available for the user's Magic Leap profile.
 	*/
-	DECLARE_DYNAMIC_DELEGATE_TwoParams(FAvailableIdentityAttributesDelegate, EMagicLeapIdentityError, ResultCode, const TArray<EMagicLeapIdentityAttribute>&, AvailableAttributes);
+	DECLARE_DYNAMIC_DELEGATE_TwoParams(FAvailableIdentityAttributesDelegate, EMagicLeapIdentityError, ResultCode, const TArray<EMagicLeapIdentityKey>&, AvailableAttributes);
 
 	/**
 	  Delegate for the result of attribute values for the user's Magic Leap profile.
@@ -39,7 +39,7 @@ public:
 	  @param ResultCode Error code when modifying the attribute values.
 	  @param AttributesUpdatedSuccessfully List of attributes whose values were successfully modified.
 	*/
-	DECLARE_DYNAMIC_DELEGATE_TwoParams(FModifyIdentityAttributeValueDelegate, EMagicLeapIdentityError, ResultCode, const TArray<EMagicLeapIdentityAttribute>&, AttributesUpdatedSuccessfully);
+	DECLARE_DYNAMIC_DELEGATE_TwoParams(FModifyIdentityAttributeValueDelegate, EMagicLeapIdentityError, ResultCode, const TArray<EMagicLeapIdentityKey>&, AttributesUpdatedSuccessfully);
 
 	/**
 	  Get the attributes available for the user's Magic Leap profile. Note that this does not request the values for these attribtues.
@@ -48,7 +48,7 @@ public:
 	  @return Error code when getting the list of available attributes.
 	*/
 	UFUNCTION(BlueprintCallable, Category = "Identity|MagicLeap")
-	EMagicLeapIdentityError GetAllAvailableAttributes(TArray<EMagicLeapIdentityAttribute>& AvailableAttributes);
+	EMagicLeapIdentityError GetAllAvailableAttributes(TArray<EMagicLeapIdentityKey>& AvailableAttributes);
 
 	/**
 	  Asynchronous call to get the attributes available for the user's Magic Leap profile. Note that this does not request the values for these attribtues.
@@ -66,7 +66,7 @@ public:
 	  @return Error code when getting the attribute values.
 	*/
 	UFUNCTION(BlueprintCallable, Category = "Identity|MagicLeap")
-	EMagicLeapIdentityError RequestAttributeValue(const TArray<EMagicLeapIdentityAttribute>& RequestedAttributeList, TArray<FMagicLeapIdentityAttribute>& RequestedAttributeValues);
+	EMagicLeapIdentityError RequestAttributeValue(const TArray<EMagicLeapIdentityKey>& RequestedAttributeList, TArray<FMagicLeapIdentityAttribute>& RequestedAttributeValues);
 
 	/**
 	  Asynchronous call to get the values for the attributes of the user's Magic Leap profile.
@@ -75,28 +75,7 @@ public:
 	  @return Error code when getting the attribute values.
 	*/
 	UFUNCTION(BlueprintCallable, Category = "Identity|MagicLeap")
-	EMagicLeapIdentityError RequestAttributeValueAsync(const TArray<EMagicLeapIdentityAttribute>& RequestedAttributeList, const FRequestIdentityAttributeValueDelegate& ResultDelegate);
-
-	/**
-	  Modify the values for the attributes of the user's Magic Leap profile.
-	  This function makes a blocking call to the cloud. You can alternatively use ModifyAttributeValueAsync() to modify the attribute values asynchronously.
-	  @param UpdatedAttributeValueList List of attributes and the values to update the cloud profile with.
-	  @param AttributesUpdatedSuccessfully Output parameter populated with the list of attributes whose values were updated successfully.
-	  @return Error code when modifying the attribute values.
-	*/
-	DEPRECATED(4.19, "This function is deprecated and will be removed in future releases.")
-	UFUNCTION(BlueprintCallable, meta = (DisplayName = "ModifyAttributeValue_DEPRECATED"), Category = "Identity|MagicLeap")
-	EMagicLeapIdentityError ModifyAttributeValue(const TArray<FMagicLeapIdentityAttribute>& UpdatedAttributeValueList, TArray<EMagicLeapIdentityAttribute>& AttributesUpdatedSuccessfully);
-
-	/**
-	  Asynchronous call to modify the values for the attributes of the user's Magic Leap profile.
-	  @param UpdatedAttributeValueList List of attributes and the values to update the cloud profile with.
-	  @param ResultDelegate Callback which reports the list of attributes whose values were successfully modified.
-	  @return Error code when modifying the attribute values.
-	*/
-	DEPRECATED(4.19, "This function is deprecated and will be removed in future releases.")
-	UFUNCTION(BlueprintCallable, meta = (DisplayName = "ModifyAttributeValueAsync_DEPRECATED"), Category = "Identity|MagicLeap")
-	EMagicLeapIdentityError ModifyAttributeValueAsync(const TArray<FMagicLeapIdentityAttribute>& UpdatedAttributeValueList, const FModifyIdentityAttributeValueDelegate& ResultDelegate);
+	EMagicLeapIdentityError RequestAttributeValueAsync(const TArray<EMagicLeapIdentityKey>& RequestedAttributeList, const FRequestIdentityAttributeValueDelegate& ResultDelegate);
 
 	/** UObjectBase interface */
 	virtual UWorld* GetWorld() const override;
diff --git a/Engine/Plugins/Lumin/MagicLeap/Source/MagicLeapIdentity/Public/MagicLeapIdentityTypes.h b/Engine/Plugins/Lumin/MagicLeap/Source/MagicLeapIdentity/Public/MagicLeapIdentityTypes.h
--- a/Engine/Plugins/Lumin/MagicLeap/Source/MagicLeapIdentity/Public/MagicLeapIdentityTypes.h
+++ b/Engine/Plugins/Lumin/MagicLeap/Source/MagicLeapIdentity/Public/MagicLeapIdentityTypes.h
@@ -8,24 +8,16 @@
 
 /** Identifies an attribute in a user profile. */
 UENUM(BlueprintType, meta=(ScriptName="MagicLeapIdentityAttributeType"))
-enum class EMagicLeapIdentityAttribute : uint8
+enum class EMagicLeapIdentityKey : uint8
 {
-	UserID,
 	GivenName,
 	FamilyName,
 	Email,
-	Status,
-	TermsAccepted,
-	Birthday,
-	Company,
-	Industry,
-	Location,
-	Tagline,
+	Bio,
 	PhoneNumber,
 	Avatar2D,
 	Avatar3D,
-	IsDeveloper,
-	Unknown
+	Unknown,
 };
 
 /** List of possible errors when calling the Identity functions. */
@@ -33,13 +25,23 @@ UENUM(BlueprintType)
 enum class EMagicLeapIdentityError : uint8
 {
 	Ok,
+	InvalidParam,
+	AllocFailed,
+	PrivilegeDenied,
 	FailedToConnectToLocalService,
 	FailedToConnectToCloudService,
 	CloudAuthentication,
 	InvalidInformationFromCloud,
-	InvalidArgument,
-	AsyncOperationNotComplete,
-	OtherError
+	NotLoggedIn,
+	ExpiredCredentials,
+	FailedToGetUserProfile,
+	Unauthorized,
+	CertificateError,
+	RejectedByCloud,
+	AlreadyLoggedIn,
+	ModifyIsNotSupported,
+	NetworkError,
+	UnspecifiedFailure,
 };
 
 /** Represents an attribute and its value of a user's profile. */
@@ -52,13 +54,13 @@ public:
 	FMagicLeapIdentityAttribute()
 	{}
 
-	FMagicLeapIdentityAttribute(EMagicLeapIdentityAttribute attribute, const FString& value)
+	FMagicLeapIdentityAttribute(EMagicLeapIdentityKey attribute, const FString& value)
 		: Attribute(attribute)
 		, Value(value)
 	{}
 
 	UPROPERTY(EditAnywhere, BlueprintReadWrite, Category = "Identity|MagicLeap")
-	EMagicLeapIdentityAttribute Attribute;
+	EMagicLeapIdentityKey Attribute;
 
 	UPROPERTY(EditAnywhere, BlueprintReadWrite, Category = "Identity|MagicLeap")
 	FString Value;
diff --git a/Engine/Plugins/Lumin/MagicLeap/Source/MagicLeapMusicService/MagicLeapMusicService.Build.cs b/Engine/Plugins/Lumin/MagicLeap/Source/MagicLeapMusicService/MagicLeapMusicService.Build.cs
new file mode 100644
--- /dev/null
+++ b/Engine/Plugins/Lumin/MagicLeap/Source/MagicLeapMusicService/MagicLeapMusicService.Build.cs
@@ -0,0 +1,55 @@
+// %BANNER_BEGIN%
+// ---------------------------------------------------------------------
+// %COPYRIGHT_BEGIN%
+//
+// Copyright (c) 2017 Magic Leap, Inc. (COMPANY) All Rights Reserved.
+// Magic Leap, Inc. Confidential and Proprietary
+//
+// NOTICE: All information contained herein is, and remains the property
+// of COMPANY. The intellectual and technical concepts contained herein
+// are proprietary to COMPANY and may be covered by U.S. and Foreign
+// Patents, patents in process, and are protected by trade secret or
+// copyright law. Dissemination of this information or reproduction of
+// this material is strictly forbidden unless prior written permission is
+// obtained from COMPANY. Access to the source code contained herein is
+// hereby forbidden to anyone except current COMPANY employees, managers
+// or contractors who have executed Confidentiality and Non-disclosure
+// agreements explicitly covering such access.
+//
+// The copyright notice above does not evidence any actual or intended
+// publication or disclosure of this source code, which includes
+// information that is confidential and/or proprietary, and is a trade
+// secret, of COMPANY. ANY REPRODUCTION, MODIFICATION, DISTRIBUTION,
+// PUBLIC PERFORMANCE, OR PUBLIC DISPLAY OF OR THROUGH USE OF THIS
+// SOURCE CODE WITHOUT THE EXPRESS WRITTEN CONSENT OF COMPANY IS
+// STRICTLY PROHIBITED, AND IN VIOLATION OF APPLICABLE LAWS AND
+// INTERNATIONAL TREATIES. THE RECEIPT OR POSSESSION OF THIS SOURCE
+// CODE AND/OR RELATED INFORMATION DOES NOT CONVEY OR IMPLY ANY RIGHTS
+// TO REPRODUCE, DISCLOSE OR DISTRIBUTE ITS CONTENTS, OR TO MANUFACTURE,
+// USE, OR SELL ANYTHING THAT IT MAY DESCRIBE, IN WHOLE OR IN PART.
+//
+// %COPYRIGHT_END%
+// --------------------------------------------------------------------
+// %BANNER_END%
+
+using System.IO;
+
+namespace UnrealBuildTool.Rules
+{
+	public class MagicLeapMusicService : ModuleRules
+	{
+		public MagicLeapMusicService(ReadOnlyTargetRules Target) : base(Target)
+		{
+			PrivateDependencyModuleNames.AddRange(
+				new string[]
+				{
+					"Core",
+					"CoreUObject",
+					"Engine",
+					"MLSDK",
+					"LuminRuntimeSettings"
+				}
+			);
+		}
+	}
+}
diff --git a/Engine/Plugins/Lumin/MagicLeap/Source/MagicLeapMusicService/Private/MagicLeapMusicService.cpp b/Engine/Plugins/Lumin/MagicLeap/Source/MagicLeapMusicService/Private/MagicLeapMusicService.cpp
new file mode 100644
--- /dev/null
+++ b/Engine/Plugins/Lumin/MagicLeap/Source/MagicLeapMusicService/Private/MagicLeapMusicService.cpp
@@ -0,0 +1,456 @@
+// %BANNER_BEGIN%
+// ---------------------------------------------------------------------
+// %COPYRIGHT_BEGIN%
+//
+// Copyright (c) 2017 Magic Leap, Inc. (COMPANY) All Rights Reserved.
+// Magic Leap, Inc. Confidential and Proprietary
+//
+// NOTICE: All information contained herein is, and remains the property
+// of COMPANY. The intellectual and technical concepts contained herein
+// are proprietary to COMPANY and may be covered by U.S. and Foreign
+// Patents, patents in process, and are protected by trade secret or
+// copyright law. Dissemination of this information or reproduction of
+// this material is strictly forbidden unless prior written permission is
+// obtained from COMPANY. Access to the source code contained herein is
+// hereby forbidden to anyone except current COMPANY employees, managers
+// or contractors who have executed Confidentiality and Non-disclosure
+// agreements explicitly covering such access.
+//
+// The copyright notice above does not evidence any actual or intended
+// publication or disclosure of this source code, which includes
+// information that is confidential and/or proprietary, and is a trade
+// secret, of COMPANY. ANY REPRODUCTION, MODIFICATION, DISTRIBUTION,
+// PUBLIC PERFORMANCE, OR PUBLIC DISPLAY OF OR THROUGH USE OF THIS
+// SOURCE CODE WITHOUT THE EXPRESS WRITTEN CONSENT OF COMPANY IS
+// STRICTLY PROHIBITED, AND IN VIOLATION OF APPLICABLE LAWS AND
+// INTERNATIONAL TREATIES. THE RECEIPT OR POSSESSION OF THIS SOURCE
+// CODE AND/OR RELATED INFORMATION DOES NOT CONVEY OR IMPLY ANY RIGHTS
+// TO REPRODUCE, DISCLOSE OR DISTRIBUTE ITS CONTENTS, OR TO MANUFACTURE,
+// USE, OR SELL ANYTHING THAT IT MAY DESCRIBE, IN WHOLE OR IN PART.
+//
+// %COPYRIGHT_END%
+// --------------------------------------------------------------------
+// %BANNER_END%
+
+#include "MagicLeapMusicService.h"
+#include "IMagicLeapMusicServicePlugin.h"
+#include "CoreMinimal.h"
+#include "Engine/Engine.h"
+#include "Modules/ModuleManager.h"
+#include "HAL/PlatformProcess.h"
+#include "Misc/Paths.h"
+#include "MagicLeapPluginUtil.h"
+
+#if WITH_MLSDK
+#include "ml_music_service.h"
+#endif //WITH_MLSDK
+
+DEFINE_LOG_CATEGORY_STATIC(LogMusicService, Display, All);
+
+class FMagicLeapMusicServicePlugin : public IMagicLeapMusicServicePlugin
+{
+public:
+	virtual void StartupModule() override
+	{
+		IModuleInterface::StartupModule();
+		APISetup.Startup();
+#if WITH_MLSDK
+		APISetup.LoadDLL(TEXT("ml_musicservice"));
+#endif //WITH_MLSDK
+	}
+
+	virtual void ShutdownModule() override
+	{
+		APISetup.Shutdown();
+		IModuleInterface::ShutdownModule();
+	}
+
+private:
+	FMagicLeapAPISetup APISetup;
+};
+
+IMPLEMENT_MODULE(FMagicLeapMusicServicePlugin, MagicLeapMusicService);
+
+//////////////////////////////////////////////////////////////////////////
+
+#if WITH_MLSDK
+MLMusicServiceShuffleState MapEnum(EMagicLeapMusicServiceShuffleState SourceEnum)
+{
+	switch (SourceEnum)
+	{
+		case EMagicLeapMusicServiceShuffleState::On:
+			return MLMusicServiceShuffleState_On;
+		case EMagicLeapMusicServiceShuffleState::Off:
+			return MLMusicServiceShuffleState_Off;
+		default:
+			break;
+	}
+	return MLMusicServiceShuffleState_Unknown;
+}
+
+MLMusicServiceRepeatState MapEnum(EMagicLeapMusicServiceRepeatState SourceEnum)
+{
+	switch (SourceEnum)
+	{
+		case EMagicLeapMusicServiceRepeatState::Off:
+			return MLMusicServiceRepeatState_Off;
+		case EMagicLeapMusicServiceRepeatState::Song:
+			return MLMusicServiceRepeatState_Song;
+		case EMagicLeapMusicServiceRepeatState::Album:
+			return MLMusicServiceRepeatState_Album;
+		default:
+			break;
+	}
+	return MLMusicServiceRepeatState_Unknown;
+}
+
+EMagicLeapMusicServiceStatus MapEnum(MLMusicServiceStatus SourceEnum)
+{
+	switch (SourceEnum)
+	{
+		case MLMusicServiceStatus_ContextChanged:
+			return EMagicLeapMusicServiceStatus::ContextChanged;
+		case MLMusicServiceStatus_Created:
+			return EMagicLeapMusicServiceStatus::Created;
+		case MLMusicServiceStatus_LoggedIn:
+			return EMagicLeapMusicServiceStatus::LoggedIn;
+		case MLMusicServiceStatus_LoggedOut:
+			return EMagicLeapMusicServiceStatus::LoggedOut;
+		case MLMusicServiceStatus_NextTrack:
+			return EMagicLeapMusicServiceStatus::NextTrack;
+		case MLMusicServiceStatus_PrevTrack:
+			return EMagicLeapMusicServiceStatus::PrevTrack;
+		case MLMusicServiceStatus_TrackChanged:
+			return EMagicLeapMusicServiceStatus::TrackChanged;
+		default:
+			break;
+	}
+	return EMagicLeapMusicServiceStatus::Unknown;
+}
+
+EMagicLeapMusicServiceError MapEnum(MLMusicServiceError SourceEnum)
+{
+	switch (SourceEnum)
+	{
+		case MLMusicServiceError_None:
+			return EMagicLeapMusicServiceError::None;
+		case MLMusicServiceError_Connectivity:
+			return EMagicLeapMusicServiceError::Connectivity;
+		case MLMusicServiceError_Timeout:
+			return EMagicLeapMusicServiceError::Timeout;
+		case MLMusicServiceError_GeneralPlayback:
+			return EMagicLeapMusicServiceError::GeneralPlayback;
+		case MLMusicServiceError_Privilege:
+			return EMagicLeapMusicServiceError::Privilege;
+		case MLMusicServiceError_ServiceSpecific:
+			return EMagicLeapMusicServiceError::ServiceSpecific;
+		case MLMusicServiceError_NoMemory:
+			return EMagicLeapMusicServiceError::NoMemory;
+		default:
+			break;
+	}
+	return EMagicLeapMusicServiceError::Unspecified;
+}
+
+EMagicLeapMusicServicePlaybackState MapEnum(MLMusicServicePlaybackState SourceEnum)
+{
+	switch (SourceEnum)
+	{
+		case MLMusicServicePlaybackState_Playing:
+			return EMagicLeapMusicServicePlaybackState::Playing;
+		case MLMusicServicePlaybackState_Paused:
+			return EMagicLeapMusicServicePlaybackState::Paused;
+		case MLMusicServicePlaybackState_Stopped:
+			return EMagicLeapMusicServicePlaybackState::Stopped;
+		case MLMusicServicePlaybackState_Error:
+			return EMagicLeapMusicServicePlaybackState::Error;
+		default:
+			break;
+	}
+	return EMagicLeapMusicServicePlaybackState::Unknown;
+}
+
+EMagicLeapMusicServiceRepeatState MapEnum(MLMusicServiceRepeatState SourceEnum)
+{
+	switch (SourceEnum)
+	{
+		case MLMusicServiceRepeatState_Off:
+			return EMagicLeapMusicServiceRepeatState::Off;
+		case MLMusicServiceRepeatState_Song:
+			return EMagicLeapMusicServiceRepeatState::Song;
+		case MLMusicServiceRepeatState_Album:
+			return EMagicLeapMusicServiceRepeatState::Album;
+		default:
+			break;
+	}
+	return EMagicLeapMusicServiceRepeatState::Unknown;
+}
+
+EMagicLeapMusicServiceShuffleState MapEnum(MLMusicServiceShuffleState SourceEnum)
+{
+	switch (SourceEnum)
+	{
+		case MLMusicServiceShuffleState_On:
+			return EMagicLeapMusicServiceShuffleState::On;
+		case MLMusicServiceShuffleState_Off:
+			return EMagicLeapMusicServiceShuffleState::Off;
+		default:
+			break;
+	}
+	return EMagicLeapMusicServiceShuffleState::Unknown;
+}
+#endif //WITH_MLSDK
+
+bool UMagicLeapMusicService::Connect(const FString& ProviderName)
+{
+#if WITH_MLSDK
+	return (MLMusicServiceConnect(TCHAR_TO_ANSI(*ProviderName)) == MLResult_Ok);
+#else
+	return false;
+#endif //WITH_MLSDK
+}
+
+bool UMagicLeapMusicService::Disconnect()
+{
+#if WITH_MLSDK
+	return MLMusicServiceDisconnect() == MLResult_Ok;
+#else
+	return false;
+#endif //WITH_MLSDK
+}
+
+bool UMagicLeapMusicService::SetAuthorizationString(const FString& AuthorizationString)
+{
+#if WITH_MLSDK
+	return MLMusicServiceSetAuthString(TCHAR_TO_ANSI(*AuthorizationString)) == MLResult_Ok;
+#else
+	return false;
+#endif //WITH_MLSDK
+}
+
+bool UMagicLeapMusicService::SetUrl(const FString& requestedUrl)
+{
+#if WITH_MLSDK
+	return (MLMusicServiceSetURL(TCHAR_TO_ANSI(*requestedUrl)) == MLResult_Ok);
+#else
+	return false;
+#endif //WITH_MLSDK
+
+}
+
+bool UMagicLeapMusicService::SetUrlList(const TArray<FString>& Playlist)
+{
+#if WITH_MLSDK
+	TArray<const char *> ptrs;
+	for (int i = 0; i < Playlist.Num(); ++ i) {
+		ptrs.Add(TCHAR_TO_ANSI(*Playlist[i]));
+	}
+	return MLMusicServiceSetPlayList(&ptrs[0], ptrs.Num()) == MLResult_Ok;
+#else
+	return false;
+#endif //WITH_MLSDK
+}
+
+bool UMagicLeapMusicService::Start()
+{
+#if WITH_MLSDK
+	return MLMusicServiceStart() == MLResult_Ok;
+#else
+	return false;
+#endif //WITH_MLSDK
+}
+
+bool UMagicLeapMusicService::Stop()
+{
+#if WITH_MLSDK
+	return MLMusicServiceStop() == MLResult_Ok;
+#else
+	return false;
+#endif //WITH_MLSDK
+}
+
+bool UMagicLeapMusicService::Pause()
+{
+#if WITH_MLSDK
+	return MLMusicServicePause() == MLResult_Ok;
+#else
+	return false;
+#endif //WITH_MLSDK
+}
+
+bool UMagicLeapMusicService::Resume()
+{
+#if WITH_MLSDK
+	return MLMusicServiceResume() == MLResult_Ok;
+#else
+	return false;
+#endif //WITH_MLSDK
+}
+
+bool UMagicLeapMusicService::Seek(int32 SeekMs)
+{
+#if WITH_MLSDK
+	return MLMusicServiceSeek(static_cast<uint32>(SeekMs)) == MLResult_Ok;
+#else
+	return false;
+#endif //WITH_MLSDK
+}
+
+bool UMagicLeapMusicService::Next()
+{
+#if WITH_MLSDK
+	return MLMusicServiceNext() == MLResult_Ok;
+#else
+	return false;
+#endif //WITH_MLSDK
+}
+
+bool UMagicLeapMusicService::Previous()
+{
+#if WITH_MLSDK
+	return MLMusicServicePrevious() == MLResult_Ok;
+#else
+	return false;
+#endif //WITH_MLSDK
+}
+
+bool UMagicLeapMusicService::SetShuffle(EMagicLeapMusicServiceShuffleState ShuffleState)
+{
+#if WITH_MLSDK
+	return MLMusicServiceSetShuffle(MapEnum(ShuffleState)) == MLResult_Ok;
+#else
+	return false;
+#endif //WITH_MLSDK
+}
+
+bool UMagicLeapMusicService::SetRepeat(EMagicLeapMusicServiceRepeatState RepeatState)
+{
+#if WITH_MLSDK
+	return MLMusicServiceSetRepeat(MapEnum(RepeatState)) == MLResult_Ok;
+#else
+	return false;
+#endif //WITH_MLSDK
+}
+
+bool UMagicLeapMusicService::SetVolume(float Volume)
+{
+#if WITH_MLSDK
+	return MLMusicServiceSetVolume(Volume) == MLResult_Ok;
+#else
+	return false;
+#endif //WITH_MLSDK
+}
+
+bool UMagicLeapMusicService::GetTrackLength(int32& TrackLength)
+{
+#if WITH_MLSDK
+	uint32 NativeTrackLength;
+	if (MLMusicServiceGetTrackLength(&NativeTrackLength) == MLResult_Ok)
+	{
+		TrackLength = static_cast<int32>(NativeTrackLength);
+		return true;
+	}
+#endif //WITH_MLSDK
+	return false;
+}
+
+bool UMagicLeapMusicService::GetCurrentPosition(int32& Position)
+{
+#if WITH_MLSDK
+	uint32 NativePosition;
+	if (MLMusicServiceGetCurrentPosition(&NativePosition) == MLResult_Ok)
+	{
+		Position = static_cast<int32>(NativePosition);
+		return true;
+	}
+#endif //WITH_MLSDK
+	return false;
+}
+
+bool UMagicLeapMusicService::GetStatus(EMagicLeapMusicServiceStatus& CurrentStatus)
+{
+#if WITH_MLSDK
+	MLMusicServiceStatus NativeCurrentStatus = MLMusicServiceStatus_Unknown;
+	if (MLMusicServiceGetStatus(&NativeCurrentStatus) == MLResult_Ok)
+	{
+		CurrentStatus = MapEnum(NativeCurrentStatus);
+		return true;
+	}
+#endif //WITH_MLSDK
+	return false;
+}
+
+bool UMagicLeapMusicService::GetLastError(EMagicLeapMusicServiceError& LastError)
+{
+#if WITH_MLSDK
+	MLMusicServiceError NativeLastError = MLMusicServiceError_None;
+	int32 ErrorCode;
+	if (MLMusicServiceGetError(&NativeLastError, &ErrorCode) == MLResult_Ok)
+	{
+		LastError = MapEnum(NativeLastError);
+		return true;
+	}
+#endif //WITH_MLSDK
+	return false;
+}
+
+bool UMagicLeapMusicService::GetPlaybackState(EMagicLeapMusicServicePlaybackState& PlaybackState)
+{
+#if WITH_MLSDK
+	MLMusicServicePlaybackState NativePlaybackState = MLMusicServicePlaybackState_Unknown;
+	if (MLMusicServiceGetPlaybackState(&NativePlaybackState) == MLResult_Ok)
+	{
+		PlaybackState = MapEnum(NativePlaybackState);
+		return true;
+	}
+#endif //WITH_MLSDK
+	return false;
+}
+
+bool UMagicLeapMusicService::GetRepeatState(EMagicLeapMusicServiceRepeatState& RepeatState)
+{
+#if WITH_MLSDK
+	MLMusicServiceRepeatState NativeRepeatState = MLMusicServiceRepeatState_Unknown;
+	if (MLMusicServiceGetRepeatState(&NativeRepeatState) == MLResult_Ok)
+	{
+		RepeatState = MapEnum(NativeRepeatState);
+		return true;
+	}
+#endif //WITH_MLSDK
+	return false;
+}
+
+bool UMagicLeapMusicService::GetShuffleState(EMagicLeapMusicServiceShuffleState& ShuffleState)
+{
+#if WITH_MLSDK
+	MLMusicServiceShuffleState NativeShuffleState = MLMusicServiceShuffleState_Unknown;
+	if (MLMusicServiceGetShuffleState(&NativeShuffleState) == MLResult_Ok)
+	{
+		ShuffleState = MapEnum(NativeShuffleState);
+		return true;
+	}
+#endif //WITH_MLSDK
+	return false;
+}
+
+bool UMagicLeapMusicService::GetMetadata(FMagicLeapMusicServiceMetadata& MetaData)
+{
+#if WITH_MLSDK
+	MLMusicServiceMetadata NativeMetaData;
+	if (MLMusicServiceGetMetadata(&NativeMetaData) == MLResult_Ok)
+	{
+		MetaData.AlbumTitle = NativeMetaData.album_title;
+		MetaData.AlbumInfoName = NativeMetaData.album_info_name;
+		MetaData.AlbumInfoUrl = NativeMetaData.album_info_url;
+		MetaData.AlbumInfoCoverUrl = NativeMetaData.album_info_cover_url;
+		MetaData.ArtistInfoName = NativeMetaData.artist_info_name;
+		MetaData.ArtistInfoUrl = NativeMetaData.artist_info_url;
+		MetaData.Length = NativeMetaData.length;
+		MetaData.Position = NativeMetaData.position;
+		(void)MLMusicServiceReleaseMetadata(&NativeMetaData);
+		return true;
+	}
+
+#endif //WITH_MLSDK
+	return false;
+}
diff --git a/Engine/Plugins/Lumin/MagicLeap/Source/MagicLeapMusicService/Public/IMagicLeapMusicServicePlugin.h b/Engine/Plugins/Lumin/MagicLeap/Source/MagicLeapMusicService/Public/IMagicLeapMusicServicePlugin.h
new file mode 100644
--- /dev/null
+++ b/Engine/Plugins/Lumin/MagicLeap/Source/MagicLeapMusicService/Public/IMagicLeapMusicServicePlugin.h
@@ -0,0 +1,67 @@
+// %BANNER_BEGIN%
+// ---------------------------------------------------------------------
+// %COPYRIGHT_BEGIN%
+//
+// Copyright (c) 2017 Magic Leap, Inc. (COMPANY) All Rights Reserved.
+// Magic Leap, Inc. Confidential and Proprietary
+//
+// NOTICE: All information contained herein is, and remains the property
+// of COMPANY. The intellectual and technical concepts contained herein
+// are proprietary to COMPANY and may be covered by U.S. and Foreign
+// Patents, patents in process, and are protected by trade secret or
+// copyright law. Dissemination of this information or reproduction of
+// this material is strictly forbidden unless prior written permission is
+// obtained from COMPANY. Access to the source code contained herein is
+// hereby forbidden to anyone except current COMPANY employees, managers
+// or contractors who have executed Confidentiality and Non-disclosure
+// agreements explicitly covering such access.
+//
+// The copyright notice above does not evidence any actual or intended
+// publication or disclosure of this source code, which includes
+// information that is confidential and/or proprietary, and is a trade
+// secret, of COMPANY. ANY REPRODUCTION, MODIFICATION, DISTRIBUTION,
+// PUBLIC PERFORMANCE, OR PUBLIC DISPLAY OF OR THROUGH USE OF THIS
+// SOURCE CODE WITHOUT THE EXPRESS WRITTEN CONSENT OF COMPANY IS
+// STRICTLY PROHIBITED, AND IN VIOLATION OF APPLICABLE LAWS AND
+// INTERNATIONAL TREATIES. THE RECEIPT OR POSSESSION OF THIS SOURCE
+// CODE AND/OR RELATED INFORMATION DOES NOT CONVEY OR IMPLY ANY RIGHTS
+// TO REPRODUCE, DISCLOSE OR DISTRIBUTE ITS CONTENTS, OR TO MANUFACTURE,
+// USE, OR SELL ANYTHING THAT IT MAY DESCRIBE, IN WHOLE OR IN PART.
+//
+// %COPYRIGHT_END%
+// --------------------------------------------------------------------
+// %BANNER_END%
+
+#pragma once
+
+#include "Modules/ModuleManager.h"
+
+/**
+* The public interface to this module.  In most cases, this interface is only public to sibling modules
+* within this plugin.
+*/
+class IMagicLeapMusicServicePlugin : public IModuleInterface
+{
+public:
+
+	/**
+	* Singleton-like access to this module's interface.  This is just for convenience!
+	* Beware of calling this during the shutdown phase, though.  Your module might have been unloaded already.
+	*
+	* @return Returns singleton instance, loading the module on demand if needed
+	*/
+	static inline IMagicLeapMusicServicePlugin& Get()
+	{
+		return FModuleManager::LoadModuleChecked<IMagicLeapMusicServicePlugin>("MagicLeapMusicService");
+	}
+
+	/**
+	* Checks to see if this module is loaded and ready.  It is only valid to call Get() if IsAvailable() returns true.
+	*
+	* @return True if the module is loaded and ready to use
+	*/
+	static inline bool IsAvailable()
+	{
+		return FModuleManager::Get().IsModuleLoaded("MagicLeapMusicService");
+	}
+};
diff --git a/Engine/Plugins/Lumin/MagicLeap/Source/MagicLeapMusicService/Public/MagicLeapMusicService.h b/Engine/Plugins/Lumin/MagicLeap/Source/MagicLeapMusicService/Public/MagicLeapMusicService.h
new file mode 100644
--- /dev/null
+++ b/Engine/Plugins/Lumin/MagicLeap/Source/MagicLeapMusicService/Public/MagicLeapMusicService.h
@@ -0,0 +1,228 @@
+// %BANNER_BEGIN%
+// ---------------------------------------------------------------------
+// %COPYRIGHT_BEGIN%
+//
+// Copyright (c) 2017 Magic Leap, Inc. (COMPANY) All Rights Reserved.
+// Magic Leap, Inc. Confidential and Proprietary
+//
+// NOTICE: All information contained herein is, and remains the property
+// of COMPANY. The intellectual and technical concepts contained herein
+// are proprietary to COMPANY and may be covered by U.S. and Foreign
+// Patents, patents in process, and are protected by trade secret or
+// copyright law. Dissemination of this information or reproduction of
+// this material is strictly forbidden unless prior written permission is
+// obtained from COMPANY. Access to the source code contained herein is
+// hereby forbidden to anyone except current COMPANY employees, managers
+// or contractors who have executed Confidentiality and Non-disclosure
+// agreements explicitly covering such access.
+//
+// The copyright notice above does not evidence any actual or intended
+// publication or disclosure of this source code, which includes
+// information that is confidential and/or proprietary, and is a trade
+// secret, of COMPANY. ANY REPRODUCTION, MODIFICATION, DISTRIBUTION,
+// PUBLIC PERFORMANCE, OR PUBLIC DISPLAY OF OR THROUGH USE OF THIS
+// SOURCE CODE WITHOUT THE EXPRESS WRITTEN CONSENT OF COMPANY IS
+// STRICTLY PROHIBITED, AND IN VIOLATION OF APPLICABLE LAWS AND
+// INTERNATIONAL TREATIES. THE RECEIPT OR POSSESSION OF THIS SOURCE
+// CODE AND/OR RELATED INFORMATION DOES NOT CONVEY OR IMPLY ANY RIGHTS
+// TO REPRODUCE, DISCLOSE OR DISTRIBUTE ITS CONTENTS, OR TO MANUFACTURE,
+// USE, OR SELL ANYTHING THAT IT MAY DESCRIBE, IN WHOLE OR IN PART.
+//
+// %COPYRIGHT_END%
+// --------------------------------------------------------------------
+// %BANNER_END%
+
+#pragma once
+
+#include "CoreMinimal.h"
+#include "UObject/ObjectMacros.h"
+#include "Kismet/BlueprintFunctionLibrary.h"
+#include "MagicLeapMusicServiceTypes.h"
+#include "MagicLeapMusicService.generated.h"
+
+/**
+Function library for the Magic Leap Music Service API.
+*/
+UCLASS(ClassGroup = MagicLeap)
+class MAGICLEAPMUSICSERVICE_API UMagicLeapMusicService : public UBlueprintFunctionLibrary
+{
+	GENERATED_BODY()
+
+public:
+	/**
+	Connects to the specified music service.
+	@param   ProviderName String containing name of provider.
+	@return  True if the operation was successful, false otherwise.
+	*/
+	UFUNCTION(BlueprintCallable, Category = "MusicService|MagicLeap")
+		static bool Connect(const FString& ProviderName);
+
+	/**
+	Disconnects from the current music service.
+	@return  True if the operation was successful, false otherwise.
+	*/
+	UFUNCTION(BlueprintCallable, Category = "MusicService|MagicLeap")
+		static bool Disconnect();
+
+	/**
+	Sets the authorization string for the current music service.
+	@param   AuthorizationString Authorization data.
+	@return  True if the operation was successful, false otherwise.
+	*/
+	UFUNCTION(BlueprintCallable, Category = "MusicService|MagicLeap")
+		static bool SetAuthorizationString(const FString& AuthorizationString);
+
+	/**
+	Plays a specified URL on the currently connected music service.
+	@param   RequestedUrl URL specifying music resource.
+	@return  True if the operation was successful, false otherwise.
+	*/
+	UFUNCTION(BlueprintCallable, Category = "MusicService|MagicLeap")
+		static bool SetUrl(const FString& RequestedUrl);
+
+	/**
+	Plays the specified URL list on the currently connected music service.
+	@param   Playlist Array of URLs specifying music resources.
+	@return  True if the operation was successful, false otherwise.
+	*/
+	UFUNCTION(BlueprintCallable, Category = "MusicService|MagicLeap")
+		static bool SetUrlList(const TArray<FString>& Playlist);
+
+	/**
+	Starts the service with the current data set.
+	@return  True if the operation was successful, false otherwise.
+	*/
+	UFUNCTION(BlueprintCallable, Category = "MusicService|MagicLeap")
+		static bool Start();
+
+	/**
+	Stops the service.
+	@return  True if the operation was successful, false otherwise.
+	*/
+	UFUNCTION(BlueprintCallable, Category = "MusicService|MagicLeap")
+		static bool Stop();
+
+	/**
+	Pauses the service.
+	@return  True if the operation was successful, false otherwise.
+	*/
+	UFUNCTION(BlueprintCallable, Category = "MusicService|MagicLeap")
+		static bool Pause();
+
+	/**
+	Resumes the service.
+	@return  True if the operation was successful, false otherwise.
+	*/
+	UFUNCTION(BlueprintCallable, Category = "MusicService|MagicLeap")
+		static bool Resume();
+
+	/**
+	Seeks to the specified (absolute/relative?) position of the current track.
+	@param   SeekMs The position (in milliseconds) to seek.
+	@return  True if the operation was successful, false otherwise.
+	*/
+	UFUNCTION(BlueprintCallable, Category = "MusicService|MagicLeap")
+		static bool Seek(int32 SeekMs);
+
+	/**
+	Advances to the next track in the active playlist.
+	@return  True if the operation was successful, false otherwise.
+	*/
+	UFUNCTION(BlueprintCallable, Category = "MusicService|MagicLeap")
+		static bool Next();
+
+	/**
+	Rewinds to the previous track in the active playlist.
+	@return  True if the operation was successful, false otherwise.
+	*/
+	UFUNCTION(BlueprintCallable, Category = "MusicService|MagicLeap")
+		static bool Previous();
+
+	/**
+	Sets the shuffle state.
+	@param   ShuffleState The desired shuffle state.
+	@return  True if the operation was successful, false otherwise.
+	*/
+	UFUNCTION(BlueprintCallable, Category = "MusicService|MagicLeap")
+		static bool SetShuffle(EMagicLeapMusicServiceShuffleState ShuffleState);
+
+	/**
+	Sets the repeat state.
+	@param   RepeatState The desired repeat state.
+	@return  True if the operation was successful, false otherwise.
+	*/
+	UFUNCTION(BlueprintCallable, Category = "MusicService|MagicLeap")
+		static bool SetRepeat(EMagicLeapMusicServiceRepeatState RepeatState);
+
+	/**
+	Sets the volume.
+	@param   Volume The desired volume.
+	@return  True if the operation was successful, false otherwise.
+	*/
+	UFUNCTION(BlueprintCallable, Category = "MusicService|MagicLeap")
+		static bool SetVolume(float Volume);
+
+	/**
+	Gets the length of the current track.
+	@param   TrackLength Length of the current track.
+	@return  True if the operation was successful, false otherwise.
+	*/
+	UFUNCTION(BlueprintCallable, Category = "MusicService|MagicLeap")
+		static bool GetTrackLength(int32& TrackLength);
+
+	/**
+	Gets the current playback position of the current track.
+	@param   Position Current position of the current track.
+	@return  True if the operation was successful, false otherwise.
+	*/
+	UFUNCTION(BlueprintCallable, Category = "MusicService|MagicLeap")
+		static bool GetCurrentPosition(int32& Position);
+
+	/**
+	Gets the current status of the music service.
+	@param   Status Current status of the music service.
+	@return  True if the operation was successful, false otherwise.
+	*/
+	UFUNCTION(BlueprintCallable, Category = "MusicService|MagicLeap")
+		static bool GetStatus(EMagicLeapMusicServiceStatus& Status);
+
+	/**
+	Gets the last error code. Call this function if another API fails.
+	@param   LastError Last error.
+	@return  True if the operation was successful, false otherwise.
+	*/
+	UFUNCTION(BlueprintCallable, Category = "MusicService|MagicLeap")
+		static bool GetLastError(EMagicLeapMusicServiceError& LastError);
+
+	/**
+	Gets the current playback state.
+	@param   PlaybackState Current playback state.
+	@return  True if the operation was successful, false otherwise.
+	*/
+	UFUNCTION(BlueprintCallable, Category = "MusicService|MagicLeap")
+		static bool GetPlaybackState(EMagicLeapMusicServicePlaybackState& PlaybackState);
+
+	/**
+	Gets the current repeat state.
+	@param   RepeatState Current repeat state.
+	@return  True if the operation was successful, false otherwise.
+	*/
+	UFUNCTION(BlueprintCallable, Category = "MusicService|MagicLeap")
+		static bool GetRepeatState(EMagicLeapMusicServiceRepeatState& RepeatState);
+
+	/**
+	Gets the current shuffle state.
+	@param   ShuffleState Current shuffle state.
+	@return  True if the operation was successful, false otherwise.
+	*/
+	UFUNCTION(BlueprintCallable, Category = "MusicService|MagicLeap")
+		static bool GetShuffleState(EMagicLeapMusicServiceShuffleState& ShuffleState);
+
+	/**
+	Gets the metadata from the current track.
+	@param   MetaData Current track's metadata.
+	@return  True if the operation was successful, false otherwise.
+	*/
+	UFUNCTION(BlueprintCallable, Category = "MusicService|MagicLeap")
+		static bool GetMetadata(FMagicLeapMusicServiceMetadata& MetaData);
+};
diff --git a/Engine/Plugins/Lumin/MagicLeap/Source/MagicLeapMusicService/Public/MagicLeapMusicServiceTypes.h b/Engine/Plugins/Lumin/MagicLeap/Source/MagicLeapMusicService/Public/MagicLeapMusicServiceTypes.h
new file mode 100644
--- /dev/null
+++ b/Engine/Plugins/Lumin/MagicLeap/Source/MagicLeapMusicService/Public/MagicLeapMusicServiceTypes.h
@@ -0,0 +1,145 @@
+// %BANNER_BEGIN%
+// ---------------------------------------------------------------------
+// %COPYRIGHT_BEGIN%
+//
+// Copyright (c) 2017 Magic Leap, Inc. (COMPANY) All Rights Reserved.
+// Magic Leap, Inc. Confidential and Proprietary
+//
+// NOTICE: All information contained herein is, and remains the property
+// of COMPANY. The intellectual and technical concepts contained herein
+// are proprietary to COMPANY and may be covered by U.S. and Foreign
+// Patents, patents in process, and are protected by trade secret or
+// copyright law. Dissemination of this information or reproduction of
+// this material is strictly forbidden unless prior written permission is
+// obtained from COMPANY. Access to the source code contained herein is
+// hereby forbidden to anyone except current COMPANY employees, managers
+// or contractors who have executed Confidentiality and Non-disclosure
+// agreements explicitly covering such access.
+//
+// The copyright notice above does not evidence any actual or intended
+// publication or disclosure of this source code, which includes
+// information that is confidential and/or proprietary, and is a trade
+// secret, of COMPANY. ANY REPRODUCTION, MODIFICATION, DISTRIBUTION,
+// PUBLIC PERFORMANCE, OR PUBLIC DISPLAY OF OR THROUGH USE OF THIS
+// SOURCE CODE WITHOUT THE EXPRESS WRITTEN CONSENT OF COMPANY IS
+// STRICTLY PROHIBITED, AND IN VIOLATION OF APPLICABLE LAWS AND
+// INTERNATIONAL TREATIES. THE RECEIPT OR POSSESSION OF THIS SOURCE
+// CODE AND/OR RELATED INFORMATION DOES NOT CONVEY OR IMPLY ANY RIGHTS
+// TO REPRODUCE, DISCLOSE OR DISTRIBUTE ITS CONTENTS, OR TO MANUFACTURE,
+// USE, OR SELL ANYTHING THAT IT MAY DESCRIBE, IN WHOLE OR IN PART.
+//
+// %COPYRIGHT_END%
+// --------------------------------------------------------------------
+// %BANNER_END%
+
+#pragma once
+
+#include "CoreMinimal.h"
+#include "UObject/ObjectMacros.h"
+#include "MagicLeapMusicServiceTypes.generated.h"
+
+/** List of possible errors when calling the music service functions. */
+UENUM(BlueprintType)
+enum class EMagicLeapMusicServiceError : uint8
+{
+	None			UMETA(DisplayName = "None"),
+	Connectivity	UMETA(DisplayName = "Connectivity"),
+	Timeout			UMETA(DisplayName = "Timeout"),
+	GeneralPlayback	UMETA(DisplayName = "General Playback"),
+	Privilege		UMETA(DisplayName = "Privilege"),
+	ServiceSpecific	UMETA(DisplayName = "Service Specific"),
+	NoMemory		UMETA(DisplayName = "No Memory"),
+	Unspecified		UMETA(DisplayName = "Unspecified"),
+};
+
+/** List of possible status of the music service. */
+UENUM(BlueprintType)
+enum class EMagicLeapMusicServiceStatus : uint8
+{
+	ContextChanged	UMETA(DisplayName = "Context Changed"),
+	Created			UMETA(DisplayName = "Created"),
+	LoggedIn		UMETA(DisplayName = "Logged In"),
+	LoggedOut		UMETA(DisplayName = "Logged Out"),
+	NextTrack		UMETA(DisplayName = "Next Track"),
+	PrevTrack		UMETA(DisplayName = "Prev Track"),
+	TrackChanged	UMETA(DisplayName = "Track Changed"),
+	Unknown			UMETA(DisplayName = "Unknown"),
+};
+
+/** List of MusicService playback options. */
+UENUM(BlueprintType)
+enum class EMagicLeapMusicServicePlaybackOption : uint8
+{
+	Playback		UMETA(DisplayName = "Playback"),
+	Repeat			UMETA(DisplayName = "Repeat"),
+	Shuffle			UMETA(DisplayName = "Shuffle"),
+	Unknown			UMETA(DisplayName = "Unknown"),
+};
+
+/** List of MusicService playback states. */
+UENUM(BlueprintType)
+enum class EMagicLeapMusicServicePlaybackState : uint8
+{
+	Playing 		UMETA(DisplayName = "Playing"),
+	Paused			UMETA(DisplayName = "Paused"),
+	Stopped			UMETA(DisplayName = "Stopped"),
+	Error			UMETA(DisplayName = "Error"),
+	Unknown			UMETA(DisplayName = "Unknown"),
+};
+
+/** List of MusicService repeat settings. */
+UENUM(BlueprintType)
+enum class EMagicLeapMusicServiceRepeatState : uint8
+{
+	Off				UMETA(DisplayName = "Off"),
+	Song			UMETA(DisplayName = "Song"),
+	Album			UMETA(DisplayName = "Album"),
+	Unknown			UMETA(DisplayName = "Unknown"),
+};
+
+/** List of MusicService shuffle settings. */
+UENUM(BlueprintType)
+enum class EMagicLeapMusicServiceShuffleState : uint8
+{
+	On				UMETA(DisplayName = "On"),
+	Off				UMETA(DisplayName = "Off"),
+	Unknown			UMETA(DisplayName = "Unknown"),
+};
+
+/** Track metadata. */
+USTRUCT(BlueprintType)
+struct FMagicLeapMusicServiceMetadata
+{
+	GENERATED_USTRUCT_BODY()
+
+	UPROPERTY(VisibleAnywhere, BlueprintReadOnly, Category = "Info")
+	FString AlbumTitle;
+
+	UPROPERTY(VisibleAnywhere, BlueprintReadOnly, Category = "Info")
+	FString AlbumInfoName;
+
+	UPROPERTY(VisibleAnywhere, BlueprintReadOnly, Category = "Info")
+	FString AlbumInfoUrl;
+
+	UPROPERTY(VisibleAnywhere, BlueprintReadOnly, Category = "Info")
+	FString AlbumInfoCoverUrl;
+
+	UPROPERTY(VisibleAnywhere, BlueprintReadOnly, Category = "Info")
+	FString ArtistInfoName;
+
+	UPROPERTY(VisibleAnywhere, BlueprintReadOnly, Category = "Info")
+	FString ArtistInfoUrl;
+
+	UPROPERTY(VisibleAnywhere, BlueprintReadOnly, Category = "Info")
+	int32 Length;
+
+	UPROPERTY(VisibleAnywhere, BlueprintReadOnly, Category = "Info")
+	int32 Position;
+
+	FMagicLeapMusicServiceMetadata() :
+	Length(0),
+	Position(0)
+	{
+	}
+};
+
diff --git a/Engine/Plugins/Lumin/MagicLeap/Source/MagicLeapSecureStorage/Private/MagicLeapSecureStorage.cpp b/Engine/Plugins/Lumin/MagicLeap/Source/MagicLeapSecureStorage/Private/MagicLeapSecureStorage.cpp
--- a/Engine/Plugins/Lumin/MagicLeap/Source/MagicLeapSecureStorage/Private/MagicLeapSecureStorage.cpp
+++ b/Engine/Plugins/Lumin/MagicLeap/Source/MagicLeapSecureStorage/Private/MagicLeapSecureStorage.cpp
@@ -41,66 +41,16 @@ IMPLEMENT_MODULE(FMagicLeapSecureStoragePlugin, MagicLeapSecureStorage);
 
 //////////////////////////////////////////////////////////////////////////
 
-template<class T>
-bool UMagicLeapSecureStorage::PutSecureBlob(const FString& Key, const T* DataToStore)
-{
-#if WITH_MLSDK
-	return MLSecureStoragePutBlob(TCHAR_TO_ANSI(*Key), reinterpret_cast<const uint8*>(DataToStore), sizeof(T));
-#else
-	return false;
-#endif //WITH_MLSDK
-}
-
 template<>
 bool UMagicLeapSecureStorage::PutSecureBlob<FString>(const FString& Key, const FString* DataToStore)
 {
 #if WITH_MLSDK
-	return MLSecureStoragePutBlob(TCHAR_TO_ANSI(*Key), reinterpret_cast<const uint8*>(TCHAR_TO_ANSI(*(*DataToStore))), DataToStore->Len() + 1);
+	return MLSecureStoragePutBlob(TCHAR_TO_ANSI(*Key), reinterpret_cast<const uint8*>(TCHAR_TO_ANSI(*(*DataToStore))), DataToStore->Len() + 1) == MLResult_Ok;
 #else
 	return false;
 #endif
 }
 
-template<class T>
-bool UMagicLeapSecureStorage::GetSecureBlob(const FString& Key, T& DataToRetrieve)
-{
-	uint8* blob = nullptr;
-	size_t blobLength = 0;
-
-	bool bResult = false;
-
-#if WITH_MLSDK
-	bResult = MLSecureStorageGetBlob(TCHAR_TO_ANSI(*Key), &blob, &blobLength);
-	if (bResult)
-	{
-		if (blob == nullptr)
-		{
-			UE_LOG(LogSecureStorage, Error, TEXT("Error retrieving secure blob with key %s. Blob was null."), *Key);
-			bResult = false;
-		}
-		else if (blobLength != sizeof(T))
-		{
-			UE_LOG(LogSecureStorage, Error, TEXT("Size of blob data %s does not match the size of requested data type. Requested size = %d vs Actual size = %d"), *Key, sizeof(T), blobLength);
-			bResult = false;
-			// Replace with library function call when it comes online.
-			free(blob);
-		}
-		else
-		{
-			DataToRetrieve = *reinterpret_cast<T*>(blob);
-			// Replace with library function call when it comes online.
-			free(blob);
-		}
-	}
-	else
-	{
-		UE_LOG(LogSecureStorage, Error, TEXT("Error retrieving secure blob with key %s"), *Key);
-	}
-#endif //WITH_MLSDK
-
-	return bResult;
-}
-
 template<>
 bool UMagicLeapSecureStorage::GetSecureBlob<FString>(const FString& Key, FString& DataToRetrieve)
 {
@@ -109,7 +59,7 @@ bool UMagicLeapSecureStorage::GetSecureBlob<FString>(const FString& Key, FString
 
 	bool bResult = false;
 #if WITH_MLSDK
-	bResult = MLSecureStorageGetBlob(TCHAR_TO_ANSI(*Key), &blob, &blobLength);
+	bResult = MLSecureStorageGetBlob(TCHAR_TO_ANSI(*Key), &blob, &blobLength) == MLResult_Ok;
 #endif //WITH_MLSDK
 	if (bResult)
 	{
@@ -122,7 +72,9 @@ bool UMagicLeapSecureStorage::GetSecureBlob<FString>(const FString& Key, FString
 		{
 			DataToRetrieve = FString(ANSI_TO_TCHAR(reinterpret_cast<ANSICHAR*>(blob)));
 			// Replace with library function call when it comes online.
-			free(blob);
+#if WITH_MLSDK
+			MLSecureStorageFreeBlobBuffer(blob);
+#endif //WITH_MLSDK
 		}
 	}
 	else
@@ -216,8 +168,55 @@ bool UMagicLeapSecureStorage::GetSecureTransform(const FString& Key, FTransform&
 bool UMagicLeapSecureStorage::DeleteSecureData(const FString& Key)
 {
 #if WITH_MLSDK
-	return MLSecureStorageDeleteBlob(TCHAR_TO_ANSI(*Key));
+	return MLSecureStorageDeleteBlob(TCHAR_TO_ANSI(*Key)) == MLResult_Ok;
+#else
+	return false;
+#endif //WITH_MLSDK
+}
+
+bool UMagicLeapSecureStorage::PutSecureBlobImpl(const FString& Key, const void* DataToStore, size_t DataTypeSize)
+{
+#if WITH_MLSDK
+	return MLSecureStoragePutBlob(TCHAR_TO_ANSI(*Key), reinterpret_cast<const uint8*>(DataToStore), DataTypeSize) == MLResult_Ok;
 #else
 	return false;
 #endif //WITH_MLSDK
 }
+
+uint8* UMagicLeapSecureStorage::GetSecureBlobImpl(const FString& Key, size_t DataTypeSize)
+{
+	uint8* blob = nullptr;
+	size_t blobLength = 0;
+
+#if WITH_MLSDK
+	MLResult result = MLSecureStorageGetBlob(TCHAR_TO_ANSI(*Key), &blob, &blobLength);
+	if (MLResult_Ok == result)
+	{
+		if (blob == nullptr)
+		{
+			UE_LOG(LogSecureStorage, Error, TEXT("Error retrieving secure blob with key %s. Blob was null."), *Key);
+			result = MLSecureStorageResult_IOFailure;
+		}
+		else if (blobLength != DataTypeSize)
+		{
+			UE_LOG(LogSecureStorage, Error, TEXT("Size of blob data %s does not match the size of requested data type. Requested size = %d vs Actual size = %d"), *Key, DataTypeSize, blobLength);
+			result = MLResult_UnspecifiedFailure;
+			MLSecureStorageFreeBlobBuffer(blob);
+		}
+	}
+	else if (result != MLSecureStorageResult_BlobNotFound)
+	{
+		UE_LOG(LogSecureStorage, Error, TEXT("Error retrieving secure blob with key %s. Error Code = %d"), *Key, static_cast<int32>(result));
+	}
+	return (MLResult_Ok == result) ? blob : nullptr;
+#else
+	return nullptr;
+#endif //WITH_MLSDK
+}
+
+void UMagicLeapSecureStorage::FreeBlobBufferImpl(uint8* Buffer)
+{
+#if WITH_MLSDK
+	MLSecureStorageFreeBlobBuffer(Buffer);
+#endif //WITH_MLSDK
+}
diff --git a/Engine/Plugins/Lumin/MagicLeap/Source/MagicLeapSecureStorage/Public/MagicLeapSecureStorage.h b/Engine/Plugins/Lumin/MagicLeap/Source/MagicLeapSecureStorage/Public/MagicLeapSecureStorage.h
--- a/Engine/Plugins/Lumin/MagicLeap/Source/MagicLeapSecureStorage/Public/MagicLeapSecureStorage.h
+++ b/Engine/Plugins/Lumin/MagicLeap/Source/MagicLeapSecureStorage/Public/MagicLeapSecureStorage.h
@@ -188,4 +188,29 @@ public:
 	*/
 	template<class T>
 	static bool GetSecureBlob(const FString& Key, T& DataToRetrieve);
+
+private:
+	static bool PutSecureBlobImpl(const FString& Key, const void* DataToStore, size_t DataTypeSize);
+
+	static uint8* GetSecureBlobImpl(const FString& Key, size_t DataTypeSize);
+
+	static void FreeBlobBufferImpl(uint8* Buffer);
 };
+
+template<class T>
+bool UMagicLeapSecureStorage::PutSecureBlob(const FString& Key, const T* DataToStore)
+{
+	return UMagicLeapSecureStorage::PutSecureBlobImpl(Key, reinterpret_cast<const void*>(DataToStore), sizeof(T));
+}
+
+template<class T>
+bool UMagicLeapSecureStorage::GetSecureBlob(const FString& Key, T& DataToRetrieve)
+{
+	uint8* blob = UMagicLeapSecureStorage::GetSecureBlobImpl(Key, sizeof(T));
+	if (blob != nullptr)
+	{
+		DataToRetrieve = *reinterpret_cast<T*>(blob);
+		UMagicLeapSecureStorage::FreeBlobBufferImpl(blob);
+	}
+	return blob != nullptr;
+}
diff --git a/Engine/Plugins/Lumin/MagicLeapMedia/Source/MagicLeapMedia/Private/MagicLeapCameraPreviewPlayer.cpp b/Engine/Plugins/Lumin/MagicLeapMedia/Source/MagicLeapMedia/Private/MagicLeapCameraPreviewPlayer.cpp
--- a/Engine/Plugins/Lumin/MagicLeapMedia/Source/MagicLeapMedia/Private/MagicLeapCameraPreviewPlayer.cpp
+++ b/Engine/Plugins/Lumin/MagicLeapMedia/Source/MagicLeapMedia/Private/MagicLeapCameraPreviewPlayer.cpp
@@ -4,6 +4,7 @@
 #include "IMagicLeapMediaModule.h"
 #include "IMediaEventSink.h"
 #include "CameraCaptureComponent.h"
+#include "ExternalTexture.h"
 
 FMagicLeapCameraPreviewPlayer::FMagicLeapCameraPreviewPlayer(IMediaEventSink& InEventSink)
 	: FMagicLeapMediaPlayer(InEventSink)
@@ -45,6 +46,16 @@ bool FMagicLeapCameraPreviewPlayer::IsLooping() const
 	return true;
 }
 
+FTimespan FMagicLeapCameraPreviewPlayer::GetTime() const
+{
+	return FTimespan::Zero();
+}
+
+void FMagicLeapCameraPreviewPlayer::RegisterExternalTexture(const FGuid& InGuid, FTextureRHIRef& InTextureRHI, FSamplerStateRHIRef& InSamplerStateRHI)
+{
+	FExternalTextureRegistry::Get().RegisterExternalTexture(InGuid, InTextureRHI, InSamplerStateRHI, FLinearColor(1.0f, 0.0f, 0.0f, 1.0f), FLinearColor(0.0f, 0.0f, 0.0f, 0.0f));
+}
+
 bool FMagicLeapCameraPreviewPlayer::SetRateOne()
 {
 	CurrentState = EMediaState::Playing;
diff --git a/Engine/Plugins/Lumin/MagicLeapMedia/Source/MagicLeapMedia/Private/MagicLeapCameraPreviewPlayer.h b/Engine/Plugins/Lumin/MagicLeapMedia/Source/MagicLeapMedia/Private/MagicLeapCameraPreviewPlayer.h
--- a/Engine/Plugins/Lumin/MagicLeapMedia/Source/MagicLeapMedia/Private/MagicLeapCameraPreviewPlayer.h
+++ b/Engine/Plugins/Lumin/MagicLeapMedia/Source/MagicLeapMedia/Private/MagicLeapCameraPreviewPlayer.h
@@ -20,10 +20,12 @@ public:
 	bool Open(const FString& Url, const IMediaOptions* Options) override;
 	bool GetVideoTrackFormat(int32 TrackIndex, int32 FormatIndex, FMediaVideoTrackFormat& OutFormat) const override;
 	bool IsLooping() const override;
+	FTimespan GetTime() const override;
 
 private:
 	bool SetRateOne() override;
 	bool GetMediaPlayerState(uint16 FlagToPoll) const override;
+	void RegisterExternalTexture(const FGuid& InGuid, FTextureRHIRef& InTextureRHI, FSamplerStateRHIRef& InSamplerStateRHI) override;
 	bool RenderThreadIsBufferAvailable(MLHandle MediaPlayerHandle) override;
 	bool RenderThreadGetNativeBuffer(const MLHandle MediaPlayerHandle, MLHandle& NativeBuffer) override;
 	bool RenderThreadReleaseNativeBuffer(const MLHandle MediaPlayerHandle, MLHandle NativeBuffer) override;
