Sizing the debug layer properly for the magic leap device (inverting the y when rendering with opengl).

#mlqdr

diff --git a/Engine/Plugins/Lumin/MagicLeap/Source/MagicLeap/Private/MagicLeapStereoLayers.cpp b/Engine/Plugins/Lumin/MagicLeap/Source/MagicLeap/Private/MagicLeapStereoLayers.cpp
--- a/Engine/Plugins/Lumin/MagicLeap/Source/MagicLeap/Private/MagicLeapStereoLayers.cpp
+++ b/Engine/Plugins/Lumin/MagicLeap/Source/MagicLeap/Private/MagicLeapStereoLayers.cpp
@@ -4,6 +4,8 @@
 #include "MagicLeapStereoLayers.h"
 #include "MagicLeapHMD.h"
 #include "SceneViewExtension.h"
+#include "RHI.h"
+#include "RHIDefinitions.h"
 
 #include "CoreMinimal.h"
 
@@ -17,3 +19,19 @@ IStereoLayers* FMagicLeapHMD::GetStereoLayers()
 	return DefaultStereoLayers.Get();
 }
 
+IStereoLayers::FLayerDesc FMagicLeapStereoLayers::GetDebugCanvasLayerDesc(FTextureRHIRef Texture)
+{
+	IStereoLayers::FLayerDesc StereoLayerDesc;
+	StereoLayerDesc.Transform = FTransform(FVector(110.f, 0.f, 0.f));
+	if (IsOpenGLPlatform(GMaxRHIShaderPlatform))
+	{
+		StereoLayerDesc.Transform.SetScale3D(FVector(1.f, 1.f, -1.f));
+	}
+	StereoLayerDesc.QuadSize = FVector2D(75.f, 40.f);
+	StereoLayerDesc.PositionType = IStereoLayers::ELayerType::FaceLocked;
+	StereoLayerDesc.ShapeType = IStereoLayers::ELayerShape::QuadLayer;
+	StereoLayerDesc.Texture = Texture;
+	StereoLayerDesc.Flags = IStereoLayers::ELayerFlags::LAYER_FLAG_TEX_CONTINUOUS_UPDATE;
+	StereoLayerDesc.Flags |= IStereoLayers::ELayerFlags::LAYER_FLAG_QUAD_PRESERVE_TEX_RATIO;
+	return StereoLayerDesc;
+}
diff --git a/Engine/Plugins/Lumin/MagicLeap/Source/MagicLeap/Private/MagicLeapStereoLayers.h b/Engine/Plugins/Lumin/MagicLeap/Source/MagicLeap/Private/MagicLeapStereoLayers.h
--- a/Engine/Plugins/Lumin/MagicLeap/Source/MagicLeap/Private/MagicLeapStereoLayers.h
+++ b/Engine/Plugins/Lumin/MagicLeap/Source/MagicLeap/Private/MagicLeapStereoLayers.h
@@ -9,4 +9,8 @@ class FMagicLeapStereoLayers : public FDefaultStereoLayers
 {
 public:
 	FMagicLeapStereoLayers(const class FAutoRegister& AutoRegister, class FHeadMountedDisplayBase* InHmd) : FDefaultStereoLayers(AutoRegister, InHmd) {}
+
+public:
+	//~ IStereoLayers interface
+	virtual FLayerDesc GetDebugCanvasLayerDesc(FTextureRHIRef Texture) override;
 };
diff --git a/Engine/Source/Runtime/Renderer/Private/SceneRendering.cpp b/Engine/Source/Runtime/Renderer/Private/SceneRendering.cpp
--- a/Engine/Source/Runtime/Renderer/Private/SceneRendering.cpp
+++ b/Engine/Source/Runtime/Renderer/Private/SceneRendering.cpp
@@ -144,7 +144,7 @@ static TAutoConsoleVariable<int32> CVarODSCapture(
 
 static TAutoConsoleVariable<int32> CVarDebugCanvasInLayer(
 	TEXT("vr.DebugCanvasInLayer"),
-	0,
+	(PLATFORM_LUMIN || PLATFORM_LUMINGL4),
 	TEXT("Experimental")
 	TEXT("0 to disable (default), 1 to enable."),
 	ECVF_Scalability | ECVF_RenderThreadSafe);
