#DEV-VR - fixing initializer order

diff --git a/Engine/Plugins/Lumin/MagicLeap/Source/MagicLeap/Private/LightingTrackingComponent.cpp b/Engine/Plugins/Lumin/MagicLeap/Source/MagicLeap/Private/LightingTrackingComponent.cpp
--- a/Engine/Plugins/Lumin/MagicLeap/Source/MagicLeap/Private/LightingTrackingComponent.cpp
+++ b/Engine/Plugins/Lumin/MagicLeap/Source/MagicLeap/Private/LightingTrackingComponent.cpp
@@ -332,9 +332,9 @@ public:
 };
 
 ULightingTrackingComponent::ULightingTrackingComponent()
-	: Impl(nullptr)
-	, UseGlobalAmbience(false)
+	: UseGlobalAmbience(false)
 	, UseColorTemp(false)
+	, Impl(nullptr)
 	//, UseDynamicAmbientCubeMap(false)
 {
 	PrimaryComponentTick.TickGroup = TG_PrePhysics;
