Mirroring CL4178961.

Removing spammy error log, per consult from ML - apparently, at this time, MLSnapshotGetTransform() can return NaNs (which we handle). It is currently expected from the API.

#mlqdr

diff --git a/Engine/Plugins/Lumin/MagicLeap/Source/MagicLeap/Private/AppFramework.cpp b/Engine/Plugins/Lumin/MagicLeap/Source/MagicLeap/Private/AppFramework.cpp
--- a/Engine/Plugins/Lumin/MagicLeap/Source/MagicLeap/Private/AppFramework.cpp
+++ b/Engine/Plugins/Lumin/MagicLeap/Source/MagicLeap/Private/AppFramework.cpp
@@ -246,7 +246,6 @@ bool FAppFramework::GetTransform(const MLCoordinateFrameUID& Id, FTransform& Out
 		OutTransform = MagicLeap::ToFTransform(transform, GetWorldToMetersScale());
 		if (OutTransform.ContainsNaN())
 		{
-			UE_LOG(LogMagicLeap, Error, TEXT("MLSnapshotGetTransform() returned an invalid transform with NaNs."));
 			OutReason = EFailReason::NaNsInTransform;
 			return false;
 		}
