UEVR-1209 Update to Magic Leap Hand Tracking API
-Switched from the Gestures API to the HandTracking API
-There is nothing like the 0 and 1 tracking point that change meaning per gesture in the new API, so I assigned the former Center/Pointer/Secondary special 1/3/5 2/4/6 to Center/IndexFingerTip/ThumbTip.
#mlqdr

diff --git a/Engine/Config/BaseEngine.ini b/Engine/Config/BaseEngine.ini
--- a/Engine/Config/BaseEngine.ini
+++ b/Engine/Config/BaseEngine.ini
@@ -583,6 +583,27 @@ bUseStreamingPause=false
 
 +FunctionRedirects=(OldName="SkeletalMeshComponent.SetAnimInstanceClass",NewName="SkeletalMeshComponent.K2_SetAnimInstanceClass")
 
+; 4.21
+
+; MagicLeap Gestures->HandTracking
++ClassRedirects=(OldName="MagicLeapGesturesFunctionLibrary",NewName="/Script/MagicLeapHandTracking.MagicLeapHandTrackingFunctionLibrary")
++EnumRedirects=(OldName="EStaticGestures",NewName="/Script/MagicLeapHandTracking.EHandTrackingGesture")
++EnumRedirects=(OldName="EGestureKeypointsFilterLevel",NewName="/Script/MagicLeapHandTracking.EHandTrackingKeypointFilterLevel")
++EnumRedirects=(OldName="EGestureRecognitionFilterLevel",NewName="/Script/MagicLeapHandTracking.EHandTrackingGestureFilterLevel")
++EnumRedirects=(OldName="EGestureTransformSpace",NewName="/Script/MagicLeapHandTracking.EGestureTransformSpace")
++FunctionRedirects=(OldName="NavigationSystemV1.SimpleMoveToLocation",NewName="AIBlueprintHelperLibrary.SimpleMoveToLocation")
++FunctionRedirects=(OldName="GetHandCenter",NewName="/Script/MagicLeapHandTracking.MagicLeapHandTrackingFunctionLibrary.GetHandCenter")
++FunctionRedirects=(OldName="GetHandPointer",NewName="/Script/MagicLeapHandTracking.MagicLeapHandTrackingFunctionLibrary.GetHandIndexFingerTip")
++FunctionRedirects=(OldName="GetHandSecondary",NewName="/Script/MagicLeapHandTracking.MagicLeapHandTrackingFunctionLibrary.GetHandThumbTip")
++FunctionRedirects=(OldName="GetHandCenterNormalized",NewName="/Script/MagicLeapHandTracking.MagicLeapHandTrackingFunctionLibrary.GetHandCenterNormalized")
++FunctionRedirects=(OldName="GetGestureKeypoints",NewName="/Script/MagicLeapHandTracking.MagicLeapHandTrackingFunctionLibrary.GetGestureKeypoints")
++FunctionRedirects=(OldName="SetConfiguration",NewName="/Script/MagicLeapHandTracking.MagicLeapHandTrackingFunctionLibrary.SetConfiguration")
++FunctionRedirects=(OldName="GetConfiguration",NewName="/Script/MagicLeapHandTracking.MagicLeapHandTrackingFunctionLibrary.GetConfiguration")
++FunctionRedirects=(OldName="SetStaticGestureConfidenceThreshold",NewName="/Script/MagicLeapHandTracking.MagicLeapHandTrackingFunctionLibrary.SetStaticGestureConfidenceThreshold")
++FunctionRedirects=(OldName="GetStaticGestureConfidenceThreshold",NewName="/Script/MagicLeapHandTracking.MagicLeapHandTrackingFunctionLibrary.GetStaticGestureConfidenceThreshold")
++FunctionRedirects=(OldName="GetHandGestureConfidence",NewName="/Script/MagicLeapHandTracking.MagicLeapHandTrackingFunctionLibrary.GetCurrentGestureConfidence")
++FunctionRedirects=(OldName="GetCurrentGesture",NewName="/Script/MagicLeapHandTracking.MagicLeapHandTrackingFunctionLibrary.GetCurrentGesture")
+
 [CoreUObject.Metadata]
 ; Note: UnrealHeaderTool should be rerun after making changes to MetadataRedirects to catch any keys specified in class headers
 +MetadataRedirects=(OldKey="K2Protected", NewKey="BlueprintProtected")
diff --git a/Engine/Plugins/Lumin/MagicLeap/MagicLeap.uplugin b/Engine/Plugins/Lumin/MagicLeap/MagicLeap.uplugin
--- a/Engine/Plugins/Lumin/MagicLeap/MagicLeap.uplugin
+++ b/Engine/Plugins/Lumin/MagicLeap/MagicLeap.uplugin
@@ -63,7 +63,7 @@
 			]
 		},
 		{
-			"Name" : "MagicLeapGestures",
+			"Name" : "MagicLeapHandTracking",
 			"Type" : "Runtime",
 			"LoadingPhase" : "PostConfigInit",
 			"WhitelistPlatforms" :
@@ -81,7 +81,8 @@
 			[
 				"Win64",
 				"Mac",
-				"Lumin"			]
+				"Lumin"
+			]
 		},
 		{
 			"Name" : "MagicLeapSecureStorage",
diff --git a/Engine/Plugins/Lumin/MagicLeap/Source/MagicLeapGestures/MagicLeapGestures.Build.cs b/Engine/Plugins/Lumin/MagicLeap/Source/MagicLeapGestures/MagicLeapGestures.Build.cs
deleted file mode 100644
diff --git a/Engine/Plugins/Lumin/MagicLeap/Source/MagicLeapGestures/Private/MagicLeapGestures.cpp b/Engine/Plugins/Lumin/MagicLeap/Source/MagicLeapGestures/Private/MagicLeapGestures.cpp
deleted file mode 100644
diff --git a/Engine/Plugins/Lumin/MagicLeap/Source/MagicLeapGestures/Private/MagicLeapGestures.h b/Engine/Plugins/Lumin/MagicLeap/Source/MagicLeapGestures/Private/MagicLeapGestures.h
deleted file mode 100644
diff --git a/Engine/Plugins/Lumin/MagicLeap/Source/MagicLeapGestures/Private/MagicLeapGesturesFunctionLibrary.cpp b/Engine/Plugins/Lumin/MagicLeap/Source/MagicLeapGestures/Private/MagicLeapGesturesFunctionLibrary.cpp
deleted file mode 100644
diff --git a/Engine/Plugins/Lumin/MagicLeap/Source/MagicLeapGestures/Public/IMagicLeapGesturesPlugin.h b/Engine/Plugins/Lumin/MagicLeap/Source/MagicLeapGestures/Public/IMagicLeapGesturesPlugin.h
deleted file mode 100644
diff --git a/Engine/Plugins/Lumin/MagicLeap/Source/MagicLeapGestures/Public/MagicLeapGestureTypes.h b/Engine/Plugins/Lumin/MagicLeap/Source/MagicLeapGestures/Public/MagicLeapGestureTypes.h
deleted file mode 100644
diff --git a/Engine/Plugins/Lumin/MagicLeap/Source/MagicLeapGestures/Public/MagicLeapGesturesFunctionLibrary.h b/Engine/Plugins/Lumin/MagicLeap/Source/MagicLeapGestures/Public/MagicLeapGesturesFunctionLibrary.h
deleted file mode 100644
diff --git a/Engine/Plugins/Lumin/MagicLeap/Source/MagicLeapHandTracking/MagicLeapHandTracking.Build.cs b/Engine/Plugins/Lumin/MagicLeap/Source/MagicLeapHandTracking/MagicLeapHandTracking.Build.cs
new file mode 100644
--- /dev/null
+++ b/Engine/Plugins/Lumin/MagicLeap/Source/MagicLeapHandTracking/MagicLeapHandTracking.Build.cs
@@ -0,0 +1,39 @@
+// Copyright 1998-2018 Epic Games, Inc. All Rights Reserved.
+
+using System.IO;
+
+namespace UnrealBuildTool.Rules
+{
+	public class MagicLeapHandTracking : ModuleRules
+	{
+		public MagicLeapHandTracking(ReadOnlyTargetRules Target)
+				: base(Target)
+		{
+			PublicDependencyModuleNames.AddRange(
+				new string[]
+				{
+					"InputDevice"
+				}
+			);
+
+			PrivateDependencyModuleNames.AddRange(
+				new string[]
+				{
+					"Core",
+					"CoreUObject",
+					"Engine",
+					"InputCore",
+					"Slate",
+					"HeadMountedDisplay",
+					"LuminRuntimeSettings",
+					"MagicLeap",
+					"MLSDK",
+					"SlateCore"
+				}
+			);
+
+			// This is not ideal but needs to be done in order to expose the private MagicLeapHMD header to this module.
+			PrivateIncludePaths.Add(Path.Combine(new string[] { ModuleDirectory, "..", "MagicLeap", "Private" }));
+		}
+	}
+}
diff --git a/Engine/Plugins/Lumin/MagicLeap/Source/MagicLeapHandTracking/Private/MagicLeapHandTracking.cpp b/Engine/Plugins/Lumin/MagicLeap/Source/MagicLeapHandTracking/Private/MagicLeapHandTracking.cpp
new file mode 100644
--- /dev/null
+++ b/Engine/Plugins/Lumin/MagicLeap/Source/MagicLeapHandTracking/Private/MagicLeapHandTracking.cpp
@@ -0,0 +1,1221 @@
+// Copyright 1998-2018 Epic Games, Inc. All Rights Reserved.
+
+#include "MagicLeapHandTracking.h"
+#include "IMagicLeapPlugin.h"
+#include "MagicLeapHMD.h"
+#include "MagicLeapMath.h"
+#include "Framework/Application/SlateApplication.h"
+#include "CoreMinimal.h"
+#include "Engine/Engine.h"
+#include "Modules/ModuleManager.h"
+#include "IMagicLeapHandTrackingPlugin.h"
+
+#define LOCTEXT_NAMESPACE "MagicLeapHandTracking"
+
+class FMagicLeapHandTrackingPlugin : public IMagicLeapHandTrackingPlugin
+{
+public:
+	FMagicLeapHandTrackingPlugin()
+		: InputDevice(nullptr)
+	{}
+
+	virtual void StartupModule() override
+	{
+		IMagicLeapHandTrackingPlugin::StartupModule();
+
+		// HACK: Generic Application might not be instantiated at this point so we create the input device with a
+		// dummy message handler. When the Generic Application creates the input device it passes a valid message
+		// handler to it which is further on used for all the controller events. This hack fixes issues caused by
+		// using a custom input device before the Generic Application has instantiated it. Eg. within BeginPlay()
+		//
+		// This also fixes the warnings that pop up on the custom input keys when the blueprint loads. Those
+		// warnings are caused because Unreal loads the bluerints before the input device has been instantiated
+		// and has added its keys, thus leading Unreal to believe that those keys don't exist. This hack causes
+		// an earlier instantiation of the input device, and consequently, the custom keys.
+		TSharedPtr<FGenericApplicationMessageHandler> DummyMessageHandler(new FGenericApplicationMessageHandler());
+		CreateInputDevice(DummyMessageHandler.ToSharedRef());
+	}
+
+	virtual void ShutdownModule() override
+	{
+		IMagicLeapHandTrackingPlugin::ShutdownModule();
+	}
+
+	virtual TSharedPtr<class IInputDevice> CreateInputDevice(const TSharedRef<FGenericApplicationMessageHandler>& InMessageHandler) override
+	{
+		if (!InputDevice.IsValid())
+		{
+			TSharedPtr<IInputDevice> HandTrackingInputDevice(new FMagicLeapHandTracking(InMessageHandler));
+			InputDevice = HandTrackingInputDevice;
+			return InputDevice;
+		}
+		else
+		{
+			InputDevice.Get()->SetMessageHandler(InMessageHandler);
+			return InputDevice;
+		}
+		return nullptr;
+	}
+
+	virtual TSharedPtr<IInputDevice> GetInputDevice() override
+	{
+		if (!InputDevice.IsValid())
+		{
+			InputDevice = CreateInputDevice(FSlateApplication::Get().GetPlatformApplication()->GetMessageHandler());
+		}
+		return InputDevice;
+	}
+
+private:
+	TSharedPtr<IInputDevice> InputDevice;
+};
+
+IMPLEMENT_MODULE(FMagicLeapHandTrackingPlugin, MagicLeapHandTracking);
+
+//////////////////////////////////////////////////////////////////////////
+
+// Partial sanity check of UE4 enums vs ML enums.  If these fail there has probably been an MLSDK update that requires an UE4 update.
+static_assert(static_cast<int32>(EHandTrackingGesture::NoHand) == static_cast<int32>(MLHandTrackingKeyPose::MLHandTrackingKeyPose_NoHand), "EHandTrackingGesture does not match api enum.");
+static_assert(static_cast<int32>(EHandTrackingKeypointFilterLevel::PredictiveSmoothing) == static_cast<int32>(MLKeypointFilterLevel::MLKeypointFilterLevel_2), "EHandTrackingKeypointFilterLevel does not match api enum.");
+static_assert(static_cast<int32>(EHandTrackingGestureFilterLevel::MoreRobustnessToFlicker) == static_cast<int32>(MLPoseFilterLevel::MLPoseFilterLevel_2), "EHandTrackingGestureFilterLevel does not match api enum.");
+
+// Left Gestures
+const FKey FMagicLeapHandTracking::FStaticHandTracking::Left_Finger("MagicLeap_Left_Finger");
+const FName FMagicLeapGestureKeyNames::Left_Finger_Name("MagicLeap_Left_Finger");
+
+const FKey FMagicLeapHandTracking::FStaticHandTracking::Left_Fist("MagicLeap_Left_Fist");
+const FName FMagicLeapGestureKeyNames::Left_Fist_Name("MagicLeap_Left_Fist");
+
+const FKey FMagicLeapHandTracking::FStaticHandTracking::Left_Pinch("MagicLeap_Left_Pinch");
+const FName FMagicLeapGestureKeyNames::Left_Pinch_Name("MagicLeap_Left_Pinch");
+
+const FKey FMagicLeapHandTracking::FStaticHandTracking::Left_Thumb("MagicLeap_Left_Thumb");
+const FName FMagicLeapGestureKeyNames::Left_Thumb_Name("MagicLeap_Left_Thumb");
+
+const FKey FMagicLeapHandTracking::FStaticHandTracking::Left_L("MagicLeap_Left_L");
+const FName FMagicLeapGestureKeyNames::Left_L_Name("MagicLeap_Left_L");
+
+const FKey FMagicLeapHandTracking::FStaticHandTracking::Left_OpenHandBack("MagicLeap_Left_OpenHandBack");
+const FName FMagicLeapGestureKeyNames::Left_OpenHandBack_Name("MagicLeap_Left_OpenHandBack");
+
+const FKey FMagicLeapHandTracking::FStaticHandTracking::Left_Ok("MagicLeap_Left_Ok");
+const FName FMagicLeapGestureKeyNames::Left_Ok_Name("MagicLeap_Left_Ok");
+
+const FKey FMagicLeapHandTracking::FStaticHandTracking::Left_C("MagicLeap_Left_C");
+const FName FMagicLeapGestureKeyNames::Left_C_Name("MagicLeap_Left_C");
+
+const FKey FMagicLeapHandTracking::FStaticHandTracking::Left_NoPose("MagicLeap_Left_NoPose");
+const FName FMagicLeapGestureKeyNames::Left_NoPose_Name("MagicLeap_Left_NoPose");
+
+const FKey FMagicLeapHandTracking::FStaticHandTracking::Left_NoHand("MagicLeap_Left_NoHand");
+const FName FMagicLeapGestureKeyNames::Left_NoHand_Name("MagicLeap_Left_NoHand");
+
+// Right Gestures
+const FKey FMagicLeapHandTracking::FStaticHandTracking::Right_Finger("MagicLeap_Right_Finger");
+const FName FMagicLeapGestureKeyNames::Right_Finger_Name("MagicLeap_Right_Finger");
+
+const FKey FMagicLeapHandTracking::FStaticHandTracking::Right_Fist("MagicLeap_Right_Fist");
+const FName FMagicLeapGestureKeyNames::Right_Fist_Name("MagicLeap_Right_Fist");
+
+const FKey FMagicLeapHandTracking::FStaticHandTracking::Right_Pinch("MagicLeap_Right_Pinch");
+const FName FMagicLeapGestureKeyNames::Right_Pinch_Name("MagicLeap_Right_Pinch");
+
+const FKey FMagicLeapHandTracking::FStaticHandTracking::Right_Thumb("MagicLeap_Right_Thumb");
+const FName FMagicLeapGestureKeyNames::Right_Thumb_Name("MagicLeap_Right_Thumb");
+
+const FKey FMagicLeapHandTracking::FStaticHandTracking::Right_L("MagicLeap_Right_L");
+const FName FMagicLeapGestureKeyNames::Right_L_Name("MagicLeap_Right_L");
+
+const FKey FMagicLeapHandTracking::FStaticHandTracking::Right_OpenHandBack("MagicLeap_Right_OpenHandBack");
+const FName FMagicLeapGestureKeyNames::Right_OpenHandBack_Name("MagicLeap_Right_OpenHandBack");
+
+const FKey FMagicLeapHandTracking::FStaticHandTracking::Right_Ok("MagicLeap_Right_Ok");
+const FName FMagicLeapGestureKeyNames::Right_Ok_Name("MagicLeap_Right_Ok");
+
+const FKey FMagicLeapHandTracking::FStaticHandTracking::Right_C("MagicLeap_Right_C");
+const FName FMagicLeapGestureKeyNames::Right_C_Name("MagicLeap_Right_C");
+
+const FKey FMagicLeapHandTracking::FStaticHandTracking::Right_NoPose("MagicLeap_Right_NoPose");
+const FName FMagicLeapGestureKeyNames::Right_NoPose_Name("MagicLeap_Right_NoPose");
+
+const FKey FMagicLeapHandTracking::FStaticHandTracking::Right_NoHand("MagicLeap_Right_NoHand");
+const FName FMagicLeapGestureKeyNames::Right_NoHand_Name("MagicLeap_Right_NoHand");
+
+
+// KeyPoint names
+#define DEFINE_KEY_POINT_NAME(name) FName FMagicLeapHandTracking::name##_Name(TEXT(#name));
+
+DEFINE_KEY_POINT_NAME(LeftHandCenter);
+
+DEFINE_KEY_POINT_NAME(LeftWristCenter);
+DEFINE_KEY_POINT_NAME(LeftWristUlnar);
+DEFINE_KEY_POINT_NAME(LeftWristRadial);
+
+DEFINE_KEY_POINT_NAME(LeftThumbTip);
+DEFINE_KEY_POINT_NAME(LeftThumbIP);
+DEFINE_KEY_POINT_NAME(LeftThumbMCP);
+DEFINE_KEY_POINT_NAME(LeftThumbCMC);
+
+DEFINE_KEY_POINT_NAME(LeftIndexFingerTip);
+DEFINE_KEY_POINT_NAME(LeftIndexFingerDIP);
+DEFINE_KEY_POINT_NAME(LeftIndexFingerPIP);
+DEFINE_KEY_POINT_NAME(LeftIndexFingerMCP);
+
+DEFINE_KEY_POINT_NAME(LeftMiddleFingerTip);
+DEFINE_KEY_POINT_NAME(LeftMiddleFingerDIP);
+DEFINE_KEY_POINT_NAME(LeftMiddleFingerPIP);
+DEFINE_KEY_POINT_NAME(LeftMiddleFingerMCP);
+
+DEFINE_KEY_POINT_NAME(LeftRingFingerTip);
+DEFINE_KEY_POINT_NAME(LeftRingFingerDIP);
+DEFINE_KEY_POINT_NAME(LeftRingFingerPIP);
+DEFINE_KEY_POINT_NAME(LeftRingFingerMCP);
+
+DEFINE_KEY_POINT_NAME(LeftPinkyFingerTip);
+DEFINE_KEY_POINT_NAME(LeftPinkyFingerDIP);
+DEFINE_KEY_POINT_NAME(LeftPinkyFingerPIP);
+DEFINE_KEY_POINT_NAME(LeftPinkyFingerMCP);
+
+
+DEFINE_KEY_POINT_NAME(RightHandCenter);
+
+DEFINE_KEY_POINT_NAME(RightWristCenter);
+DEFINE_KEY_POINT_NAME(RightWristUlnar);
+DEFINE_KEY_POINT_NAME(RightWristRadial);
+
+DEFINE_KEY_POINT_NAME(RightThumbTip);
+DEFINE_KEY_POINT_NAME(RightThumbIP);
+DEFINE_KEY_POINT_NAME(RightThumbMCP);
+DEFINE_KEY_POINT_NAME(RightThumbCMC);
+
+DEFINE_KEY_POINT_NAME(RightIndexFingerTip);
+DEFINE_KEY_POINT_NAME(RightIndexFingerDIP);
+DEFINE_KEY_POINT_NAME(RightIndexFingerPIP);
+DEFINE_KEY_POINT_NAME(RightIndexFingerMCP);
+
+DEFINE_KEY_POINT_NAME(RightMiddleFingerTip);
+DEFINE_KEY_POINT_NAME(RightMiddleFingerDIP);
+DEFINE_KEY_POINT_NAME(RightMiddleFingerPIP);
+DEFINE_KEY_POINT_NAME(RightMiddleFingerMCP);
+
+DEFINE_KEY_POINT_NAME(RightRingFingerTip);
+DEFINE_KEY_POINT_NAME(RightRingFingerDIP);
+DEFINE_KEY_POINT_NAME(RightRingFingerPIP);
+DEFINE_KEY_POINT_NAME(RightRingFingerMCP);
+
+DEFINE_KEY_POINT_NAME(RightPinkyFingerTip);
+DEFINE_KEY_POINT_NAME(RightPinkyFingerDIP);
+DEFINE_KEY_POINT_NAME(RightPinkyFingerPIP);
+DEFINE_KEY_POINT_NAME(RightPinkyFingerMCP);
+
+#if WITH_MLSDK
+EHandTrackingKeypointFilterLevel MLToUnrealKeypointsFilterLevel(MLKeypointFilterLevel FilterLevel)
+{
+	switch (FilterLevel)
+	{
+	case MLKeypointFilterLevel_0:
+		return EHandTrackingKeypointFilterLevel::NoFilter;
+	case MLKeypointFilterLevel_1:
+		return EHandTrackingKeypointFilterLevel::SimpleSmoothing;
+	case MLKeypointFilterLevel_2:
+		return EHandTrackingKeypointFilterLevel::PredictiveSmoothing;
+	}
+
+	return EHandTrackingKeypointFilterLevel::NoFilter;
+}
+
+EHandTrackingGestureFilterLevel MLToUnrealGestureFilterLevel(MLPoseFilterLevel FilterLevel)
+{
+	switch (FilterLevel)
+	{
+	case MLPoseFilterLevel_0:
+		return EHandTrackingGestureFilterLevel::NoFilter;
+	case MLPoseFilterLevel_1:
+		return EHandTrackingGestureFilterLevel::SlightRobustnessToFlicker;
+	case MLPoseFilterLevel_2:
+		return EHandTrackingGestureFilterLevel::MoreRobustnessToFlicker;
+	}
+
+	return EHandTrackingGestureFilterLevel::NoFilter;
+}
+
+MLKeypointFilterLevel UnrealToMLKeypointsFilterLevel(EHandTrackingKeypointFilterLevel FilterLevel)
+{
+	switch (FilterLevel)
+	{
+	case EHandTrackingKeypointFilterLevel::NoFilter:
+		return MLKeypointFilterLevel_0;
+	case EHandTrackingKeypointFilterLevel::SimpleSmoothing:
+		return MLKeypointFilterLevel_1;
+	case EHandTrackingKeypointFilterLevel::PredictiveSmoothing:
+		return MLKeypointFilterLevel_2;
+	}
+
+	return MLKeypointFilterLevel_0;
+}
+
+MLPoseFilterLevel UnrealToMLGestureFilterLevel(EHandTrackingGestureFilterLevel FilterLevel)
+{
+	switch (FilterLevel)
+	{
+	case EHandTrackingGestureFilterLevel::NoFilter:
+		return MLPoseFilterLevel_0;
+	case EHandTrackingGestureFilterLevel::SlightRobustnessToFlicker:
+		return MLPoseFilterLevel_1;
+	case EHandTrackingGestureFilterLevel::MoreRobustnessToFlicker:
+		return MLPoseFilterLevel_2;
+	}
+
+	return MLPoseFilterLevel_0;
+}
+#endif //WITH_MLSDK
+
+FMagicLeapHandTracking::FMagicLeapHandTracking(const TSharedRef<FGenericApplicationMessageHandler>& InMessageHandler)
+	: MessageHandler(InMessageHandler)
+	, DeviceIndex(0)
+	, HandTrackingPipelineEnabled(true)
+#if WITH_MLSDK
+	, HandTracker(ML_INVALID_HANDLE)
+#endif //WITH_MLSDK
+	, bIsHandTrackingStateValid(false)
+{
+#if WITH_MLSDK
+	// Initialize gesture data to default values.
+	HandTrackingDatas[0].left_hand_state.keypose = MLHandTrackingKeyPose_NoHand;
+	HandTrackingDatas[0].right_hand_state.keypose = MLHandTrackingKeyPose_NoHand;
+	HandTrackingDatas[1].left_hand_state.keypose = MLHandTrackingKeyPose_NoHand;
+	HandTrackingDatas[1].right_hand_state.keypose = MLHandTrackingKeyPose_NoHand;
+
+	GestureConfidenceThresholds.AddZeroed(static_cast<int32>(MLHandTrackingKeyPose_Count));
+#endif //WITH_MLSDK
+
+	BuildKeypointMaps();
+
+	// Register "MotionController" modular feature manually
+	IModularFeatures::Get().RegisterModularFeature(GetModularFeatureName(), this);
+	AddKeys();
+
+	// We're implicitly requiring that the MagicLeapPlugin has been loaded and
+	// initialized at this point.
+	auto HMD = IMagicLeapPlugin::Get().GetHMD().Pin();
+	if (HMD.IsValid())
+	{
+		HMD->RegisterMagicLeapInputDevice(this);
+	}
+}
+
+FMagicLeapHandTracking::~FMagicLeapHandTracking()
+{
+	// Normally, the MagicLeapPlugin will be around during unload,
+	// but it isn't an assumption that we should make.
+	if (IMagicLeapPlugin::IsAvailable())
+	{
+		auto HMD = IMagicLeapPlugin::Get().GetHMD().Pin();
+		if (HMD.IsValid())
+		{
+			HMD->UnregisterMagicLeapInputDevice(this);
+		}
+	}
+
+	Disable();
+
+	// Unregister "MotionController" modular feature manually
+	IModularFeatures::Get().UnregisterModularFeature(GetModularFeatureName(), this);
+}
+
+void FMagicLeapHandTracking::BuildKeypointMaps()
+{
+	SourceToTransformMap.Reserve(48);
+
+	SourceToTransformMap.Add(LeftHandCenter_Name, &LeftHand.HandCenter);
+
+	SourceToTransformMap.Add(LeftWristCenter_Name, &LeftHand.Wrist.Center);
+	SourceToTransformMap.Add(LeftWristUlnar_Name, &LeftHand.Wrist.Ulnar);
+	SourceToTransformMap.Add(LeftWristRadial_Name, &LeftHand.Wrist.Radial);
+
+	SourceToTransformMap.Add(LeftThumbTip_Name, &LeftHand.Thumb.Tip);
+	SourceToTransformMap.Add(LeftThumbIP_Name, &LeftHand.Thumb.DIP);
+	SourceToTransformMap.Add(LeftThumbMCP_Name, &LeftHand.Thumb.PIP);
+	SourceToTransformMap.Add(LeftThumbCMC_Name, &LeftHand.Thumb.MCP);
+
+	SourceToTransformMap.Add(LeftIndexFingerTip_Name, &LeftHand.IndexFinger.Tip);
+	SourceToTransformMap.Add(LeftIndexFingerDIP_Name, &LeftHand.IndexFinger.DIP);
+	SourceToTransformMap.Add(LeftIndexFingerPIP_Name, &LeftHand.IndexFinger.PIP);
+	SourceToTransformMap.Add(LeftIndexFingerMCP_Name, &LeftHand.IndexFinger.MCP);
+
+	SourceToTransformMap.Add(LeftMiddleFingerTip_Name, &LeftHand.MiddleFinger.Tip);
+	SourceToTransformMap.Add(LeftMiddleFingerDIP_Name, &LeftHand.MiddleFinger.DIP);
+	SourceToTransformMap.Add(LeftMiddleFingerPIP_Name, &LeftHand.MiddleFinger.PIP);
+	SourceToTransformMap.Add(LeftMiddleFingerMCP_Name, &LeftHand.MiddleFinger.MCP);
+
+	SourceToTransformMap.Add(LeftRingFingerTip_Name, &LeftHand.RingFinger.Tip);
+	SourceToTransformMap.Add(LeftRingFingerDIP_Name, &LeftHand.RingFinger.DIP);
+	SourceToTransformMap.Add(LeftRingFingerPIP_Name, &LeftHand.RingFinger.PIP);
+	SourceToTransformMap.Add(LeftRingFingerMCP_Name, &LeftHand.RingFinger.MCP);
+
+	SourceToTransformMap.Add(LeftPinkyFingerTip_Name, &LeftHand.PinkyFinger.Tip);
+	SourceToTransformMap.Add(LeftPinkyFingerDIP_Name, &LeftHand.PinkyFinger.DIP);
+	SourceToTransformMap.Add(LeftPinkyFingerPIP_Name, &LeftHand.PinkyFinger.PIP);
+	SourceToTransformMap.Add(LeftPinkyFingerMCP_Name, &LeftHand.PinkyFinger.MCP);
+
+
+	SourceToTransformMap.Add(RightHandCenter_Name, &RightHand.HandCenter);
+
+	SourceToTransformMap.Add(RightWristCenter_Name, &RightHand.Wrist.Center);
+	SourceToTransformMap.Add(RightWristUlnar_Name, &RightHand.Wrist.Ulnar);
+	SourceToTransformMap.Add(RightWristRadial_Name, &RightHand.Wrist.Radial);
+
+	SourceToTransformMap.Add(RightThumbTip_Name, &RightHand.Thumb.Tip);
+	SourceToTransformMap.Add(RightThumbIP_Name, &RightHand.Thumb.DIP);
+	SourceToTransformMap.Add(RightThumbMCP_Name, &RightHand.Thumb.PIP);
+	SourceToTransformMap.Add(RightThumbCMC_Name, &RightHand.Thumb.MCP);
+
+	SourceToTransformMap.Add(RightIndexFingerTip_Name, &RightHand.IndexFinger.Tip);
+	SourceToTransformMap.Add(RightIndexFingerDIP_Name, &RightHand.IndexFinger.DIP);
+	SourceToTransformMap.Add(RightIndexFingerPIP_Name, &RightHand.IndexFinger.PIP);
+	SourceToTransformMap.Add(RightIndexFingerMCP_Name, &RightHand.IndexFinger.MCP);
+
+	SourceToTransformMap.Add(RightMiddleFingerTip_Name, &RightHand.MiddleFinger.Tip);
+	SourceToTransformMap.Add(RightMiddleFingerDIP_Name, &RightHand.MiddleFinger.DIP);
+	SourceToTransformMap.Add(RightMiddleFingerPIP_Name, &RightHand.MiddleFinger.PIP);
+	SourceToTransformMap.Add(RightMiddleFingerMCP_Name, &RightHand.MiddleFinger.MCP);
+
+	SourceToTransformMap.Add(RightRingFingerTip_Name, &RightHand.RingFinger.Tip);
+	SourceToTransformMap.Add(RightRingFingerDIP_Name, &RightHand.RingFinger.DIP);
+	SourceToTransformMap.Add(RightRingFingerPIP_Name, &RightHand.RingFinger.PIP);
+	SourceToTransformMap.Add(RightRingFingerMCP_Name, &RightHand.RingFinger.MCP);
+
+	SourceToTransformMap.Add(RightPinkyFingerTip_Name, &RightHand.PinkyFinger.Tip);
+	SourceToTransformMap.Add(RightPinkyFingerDIP_Name, &RightHand.PinkyFinger.DIP);
+	SourceToTransformMap.Add(RightPinkyFingerPIP_Name, &RightHand.PinkyFinger.PIP);
+	SourceToTransformMap.Add(RightPinkyFingerMCP_Name, &RightHand.PinkyFinger.MCP);
+
+
+
+	SourceToHandMap.Add(LeftHandCenter_Name, &LeftHand);
+
+	SourceToHandMap.Add(LeftWristCenter_Name, &LeftHand);
+	SourceToHandMap.Add(LeftWristUlnar_Name, &LeftHand);
+	SourceToHandMap.Add(LeftWristRadial_Name, &LeftHand);
+
+	SourceToHandMap.Add(LeftThumbTip_Name, &LeftHand);
+	SourceToHandMap.Add(LeftThumbIP_Name, &LeftHand);
+	SourceToHandMap.Add(LeftThumbMCP_Name, &LeftHand);
+	SourceToHandMap.Add(LeftThumbCMC_Name, &LeftHand);
+
+	SourceToHandMap.Add(LeftIndexFingerTip_Name, &LeftHand);
+	SourceToHandMap.Add(LeftIndexFingerDIP_Name, &LeftHand);
+	SourceToHandMap.Add(LeftIndexFingerPIP_Name, &LeftHand);
+	SourceToHandMap.Add(LeftIndexFingerMCP_Name, &LeftHand);
+
+	SourceToHandMap.Add(LeftMiddleFingerTip_Name, &LeftHand);
+	SourceToHandMap.Add(LeftMiddleFingerDIP_Name, &LeftHand);
+	SourceToHandMap.Add(LeftMiddleFingerPIP_Name, &LeftHand);
+	SourceToHandMap.Add(LeftMiddleFingerMCP_Name, &LeftHand);
+
+	SourceToHandMap.Add(LeftRingFingerTip_Name, &LeftHand);
+	SourceToHandMap.Add(LeftRingFingerDIP_Name, &LeftHand);
+	SourceToHandMap.Add(LeftRingFingerPIP_Name, &LeftHand);
+	SourceToHandMap.Add(LeftRingFingerMCP_Name, &LeftHand);
+
+	SourceToHandMap.Add(LeftPinkyFingerTip_Name, &LeftHand);
+	SourceToHandMap.Add(LeftPinkyFingerDIP_Name, &LeftHand);
+	SourceToHandMap.Add(LeftPinkyFingerPIP_Name, &LeftHand);
+	SourceToHandMap.Add(LeftPinkyFingerMCP_Name, &LeftHand);
+
+
+	SourceToHandMap.Add(RightHandCenter_Name, &RightHand);
+
+	SourceToHandMap.Add(RightWristCenter_Name, &RightHand);
+	SourceToHandMap.Add(RightWristUlnar_Name, &RightHand);
+	SourceToHandMap.Add(RightWristRadial_Name, &RightHand);
+
+	SourceToHandMap.Add(RightThumbTip_Name, &RightHand);
+	SourceToHandMap.Add(RightThumbIP_Name, &RightHand);
+	SourceToHandMap.Add(RightThumbMCP_Name, &RightHand);
+	SourceToHandMap.Add(RightThumbCMC_Name, &RightHand);
+
+	SourceToHandMap.Add(RightIndexFingerTip_Name, &RightHand);
+	SourceToHandMap.Add(RightIndexFingerDIP_Name, &RightHand);
+	SourceToHandMap.Add(RightIndexFingerPIP_Name, &RightHand);
+	SourceToHandMap.Add(RightIndexFingerMCP_Name, &RightHand);
+
+	SourceToHandMap.Add(RightMiddleFingerTip_Name, &RightHand);
+	SourceToHandMap.Add(RightMiddleFingerDIP_Name, &RightHand);
+	SourceToHandMap.Add(RightMiddleFingerPIP_Name, &RightHand);
+	SourceToHandMap.Add(RightMiddleFingerMCP_Name, &RightHand);
+
+	SourceToHandMap.Add(RightRingFingerTip_Name, &RightHand);
+	SourceToHandMap.Add(RightRingFingerDIP_Name, &RightHand);
+	SourceToHandMap.Add(RightRingFingerPIP_Name, &RightHand);
+	SourceToHandMap.Add(RightRingFingerMCP_Name, &RightHand);
+
+	SourceToHandMap.Add(RightPinkyFingerTip_Name, &RightHand);
+	SourceToHandMap.Add(RightPinkyFingerDIP_Name, &RightHand);
+	SourceToHandMap.Add(RightPinkyFingerPIP_Name, &RightHand);
+	SourceToHandMap.Add(RightPinkyFingerMCP_Name, &RightHand);
+
+}
+
+const FTransform* FMagicLeapHandTracking::FindTransformBySource(FName SourceName) const
+{
+	return SourceToTransformMap.FindRef(SourceName);
+}
+
+const FMagicLeapHandTracking::FHandState* FMagicLeapHandTracking::FindHandBySource(FName SourceName) const
+{
+	return SourceToHandMap.FindRef(SourceName);
+}
+
+bool FMagicLeapHandTracking::GetControllerOrientationAndPosition(const int32 ControllerIndex, const FName MotionSource, FRotator& OutOrientation, FVector& OutPosition, float WorldToMetersScale) const
+{
+	bool bFound = false;
+	bool bControllerTracked = false;
+	if (ControllerIndex == DeviceIndex)
+	{
+		const FTransform* const ControllerTransform = FindTransformBySource(MotionSource);
+
+		if (ControllerTransform)
+		{
+			OutPosition = ControllerTransform->GetLocation();
+			OutOrientation = ControllerTransform->GetRotation().Rotator();
+
+			bFound = true;
+			const FHandState* HandState = SourceToHandMap.FindRef(MotionSource);
+			check(HandState);
+			bControllerTracked = HandState->IsValid();
+		}
+	}
+
+	if (bFound)
+	{
+		return bControllerTracked;
+	}
+	
+	// Then call super to handle a few of the default labels, for backward compatibility
+	return FXRMotionControllerBase::GetControllerOrientationAndPosition(ControllerIndex, MotionSource, OutOrientation, OutPosition, WorldToMetersScale);
+}
+
+bool FMagicLeapHandTracking::GetControllerOrientationAndPosition(const int32 ControllerIndex, const EControllerHand DeviceHand, FRotator& OutOrientation, FVector& OutPosition, float WorldToMetersScale) const
+{
+	bool bControllerTracked = false;
+	if (ControllerIndex == DeviceIndex)
+	{
+		if (GetControllerTrackingStatus(ControllerIndex, DeviceHand) != ETrackingStatus::NotTracked)
+		{
+			const FTransform* ControllerTransform = nullptr;
+			switch (DeviceHand)
+			{
+			case EControllerHand::Special_1:
+				ControllerTransform = &LeftHand.Wrist.Center;
+				break;
+			case EControllerHand::Special_3:
+				ControllerTransform = &LeftHand.IndexFinger.Tip;
+				break;
+			case EControllerHand::Special_5:
+				ControllerTransform = &LeftHand.Thumb.Tip;
+				break;
+			case EControllerHand::Special_2:
+				ControllerTransform = &RightHand.Wrist.Center;
+				break;
+			case EControllerHand::Special_4:
+				ControllerTransform = &RightHand.IndexFinger.Tip;
+				break;
+			case EControllerHand::Special_6:
+				ControllerTransform = &RightHand.Thumb.Tip;
+				break;
+			default:
+				check(false);
+			}
+
+			check(ControllerTransform);
+			OutPosition = ControllerTransform->GetLocation();
+			OutOrientation = ControllerTransform->GetRotation().Rotator();
+
+			bControllerTracked = true;
+		}
+	}
+
+	return bControllerTracked;
+}
+
+// Warning: this is only correct for the legacy motion source names.
+ETrackingStatus FMagicLeapHandTracking::GetControllerTrackingStatus(const int32 ControllerIndex, const EControllerHand DeviceHand) const
+{
+	ETrackingStatus status = ETrackingStatus::NotTracked;
+	if (bIsHandTrackingStateValid)
+	{
+		if (DeviceHand == EControllerHand::Special_1 || DeviceHand == EControllerHand::Special_3 || DeviceHand == EControllerHand::Special_5)
+		{
+			if (LeftHand.Gesture != EHandTrackingGesture::NoHand)
+			{
+				status = ETrackingStatus::Tracked;
+			}
+		}
+		else if (DeviceHand == EControllerHand::Special_2 || DeviceHand == EControllerHand::Special_4 || DeviceHand == EControllerHand::Special_6)
+		{
+			if (RightHand.Gesture != EHandTrackingGesture::NoHand)
+			{
+				status = ETrackingStatus::Tracked;
+			}
+		}
+	}
+
+	return status;
+}
+
+FName FMagicLeapHandTracking::GetMotionControllerDeviceTypeName() const
+{
+	const static FName DefaultName(TEXT("MagicLeapHandTracking"));
+	return DefaultName;
+}
+
+void FMagicLeapHandTracking::EnumerateSources(TArray<FMotionControllerSource>& SourcesOut) const
+{
+	// Exposing only the keypoints that are actually tracked in MLSDK 0.15.0 RC5
+
+	SourcesOut.Add(LeftThumbTip_Name);
+	//SourcesOut.Add(LeftThumbIP_Name);
+	SourcesOut.Add(LeftThumbMCP_Name);
+	//SourcesOut.Add(LeftThumbCMC_Name);
+
+	SourcesOut.Add(LeftIndexFingerTip_Name);
+	//SourcesOut.Add(LeftIndexFingerDIP_Name);
+	SourcesOut.Add(LeftIndexFingerPIP_Name);
+	SourcesOut.Add(LeftIndexFingerMCP_Name);
+
+	//SourcesOut.Add(LeftMiddleFingerTip_Name);
+	//SourcesOut.Add(LeftMiddleFingerDIP_Name);
+	//SourcesOut.Add(LeftMiddleFingerPIP_Name);
+	SourcesOut.Add(LeftMiddleFingerMCP_Name);
+
+	//SourcesOut.Add(LeftRingFingerTip_Name);
+	//SourcesOut.Add(LeftRingFingerDIP_Name);
+	//SourcesOut.Add(LeftRingFingerPIP_Name);
+	//SourcesOut.Add(LeftRingFingerMCP_Name);
+
+	//SourcesOut.Add(LeftPinkyFingerTip_Name);
+	//SourcesOut.Add(LeftPinkyFingerDIP_Name);
+	//SourcesOut.Add(LeftPinkyFingerPIP_Name);
+	//SourcesOut.Add(LeftPinkyFingerMCP_Name);
+
+	SourcesOut.Add(LeftWristCenter_Name);
+	//SourcesOut.Add(LeftWristUlnar_Name);
+	//SourcesOut.Add(LeftWristRadial_Name);
+
+	SourcesOut.Add(LeftHandCenter_Name);
+
+
+	SourcesOut.Add(RightThumbTip_Name);
+	//SourcesOut.Add(RightThumbIP_Name);
+	SourcesOut.Add(RightThumbMCP_Name);
+	//SourcesOut.Add(RightThumbCMC_Name);
+
+	SourcesOut.Add(RightIndexFingerTip_Name);
+	//SourcesOut.Add(RightIndexFingerDIP_Name);
+	SourcesOut.Add(RightIndexFingerPIP_Name);
+	SourcesOut.Add(RightIndexFingerMCP_Name);
+
+	//SourcesOut.Add(RightMiddleFingerTip_Name);
+	//SourcesOut.Add(RightMiddleFingerDIP_Name);
+	//SourcesOut.Add(RightMiddleFingerPIP_Name);
+	SourcesOut.Add(RightMiddleFingerMCP_Name);
+
+	//SourcesOut.Add(RightRingFingerTip_Name);
+	//SourcesOut.Add(RightRingFingerDIP_Name);
+	//SourcesOut.Add(RightRingFingerPIP_Name);
+	//SourcesOut.Add(RightRingFingerMCP_Name);
+
+	//SourcesOut.Add(RightPinkyFingerTip_Name);
+	//SourcesOut.Add(RightPinkyFingerDIP_Name);
+	//SourcesOut.Add(RightPinkyFingerPIP_Name);
+	//SourcesOut.Add(RightPinkyFingerMCP_Name);
+	SourcesOut.Add(RightWristCenter_Name);
+	//SourcesOut.Add(RightWristUlnar_Name);
+	//SourcesOut.Add(RightWristRadial_Name);
+
+	SourcesOut.Add(RightHandCenter_Name);
+}
+
+void FMagicLeapHandTracking::Tick(float DeltaTime)
+{
+	UpdateTrackerData();
+}
+
+void FMagicLeapHandTracking::SendControllerEvents()
+{
+#if WITH_MLSDK
+	if (bIsHandTrackingStateValid)
+	{
+		const MLHandTrackingData& CurrentHandTrackingData = GetCurrentHandTrackingData();
+		const MLHandTrackingData& OldHandTrackingData = GetPreviousHandTrackingData();
+
+		SendControllerEventsForHand(CurrentHandTrackingData.left_hand_state, OldHandTrackingData.left_hand_state, LeftStaticGestureMap);
+		SendControllerEventsForHand(CurrentHandTrackingData.right_hand_state, OldHandTrackingData.right_hand_state, RightStaticGestureMap);
+	}
+#endif //WITH_MLSDK
+}
+
+#if WITH_MLSDK
+void FMagicLeapHandTracking::SendControllerEventsForHand(const MLHandTrackingHandState& NewHandState, const MLHandTrackingHandState& OldHandState, const TArray<FName>& GestureMap)
+{
+	const int32 GestureIndex = static_cast<int32>(NewHandState.keypose);
+	const float OldConfidence = OldHandState.keypose_confidence[OldHandState.keypose];
+	const float NewConfidence = NewHandState.keypose_confidence[NewHandState.keypose];
+	if (NewHandState.keypose != OldHandState.keypose)
+	{
+		FMagicLeapHMD::EnableInput EnableInputFromHMD;
+		// fixes unreferenced parameter error for Windows package builds.
+		(void)EnableInputFromHMD;
+		MessageHandler->OnControllerButtonReleased(RightStaticGestureMap[OldHandState.keypose], DeviceIndex, false);
+		if (GestureConfidenceThresholds[GestureIndex] <= NewConfidence)
+		{
+			MessageHandler->OnControllerButtonPressed(GestureMap[GestureIndex], DeviceIndex, false);
+		}
+	}
+	else if (OldConfidence < GestureConfidenceThresholds[GestureIndex] && NewConfidence >= GestureConfidenceThresholds[GestureIndex])
+	{
+		FMagicLeapHMD::EnableInput EnableInputFromHMD;
+		// fixes unreferenced parameter error for Windows package builds.
+		(void)EnableInputFromHMD;
+		MessageHandler->OnControllerButtonPressed(GestureMap[GestureIndex], DeviceIndex, false);
+	}
+}
+#endif //WITH_MLSDK
+
+void FMagicLeapHandTracking::SetMessageHandler(const TSharedRef<FGenericApplicationMessageHandler>& InMessageHandler)
+{
+	MessageHandler = InMessageHandler;
+}
+
+bool FMagicLeapHandTracking::Exec(UWorld* InWorld, const TCHAR* Cmd, FOutputDevice& Ar)
+{
+	return false;
+}
+
+bool FMagicLeapHandTracking::IsGamepadAttached() const
+{
+#if WITH_MLSDK
+	return MLHandleIsValid(HandTracker);
+#else
+	return false;
+#endif //WITH_MLSDK
+}
+
+void FMagicLeapHandTracking::Enable()
+{
+	checkf(false, TEXT("FMagicLeapHandTracking::Enable is not supported! Check 'SupportsExplicitEnable()' first!"));
+}
+
+bool FMagicLeapHandTracking::SupportsExplicitEnable() const
+{
+	return false;
+}
+
+void FMagicLeapHandTracking::Disable()
+{
+#if WITH_MLSDK
+	if (MLHandleIsValid(HandTracker))
+	{
+		if (MLResult_Ok == MLHandTrackingDestroy(HandTracker))
+		{
+			HandTracker = ML_INVALID_HANDLE;
+		}
+		else
+		{
+			UE_LOG(LogMagicLeapHandTracking, Error, TEXT("Error destroying hand tracker."));
+		}
+	}
+#endif
+}
+
+const FMagicLeapHandTracking::FHandState& FMagicLeapHandTracking::GetLeftHandState() const
+{
+	return LeftHand;
+}
+
+const FMagicLeapHandTracking::FHandState& FMagicLeapHandTracking::GetRightHandState() const
+{
+	return RightHand;
+}
+
+bool FMagicLeapHandTracking::IsHandTrackingStateValid() const
+{
+	return bIsHandTrackingStateValid;
+}
+
+bool FMagicLeapHandTracking::SetConfiguration(bool bTrackingEnabled, const TArray<EHandTrackingGesture>& ActiveKeyPoses, EHandTrackingKeypointFilterLevel KeypointsFilterLevel, EHandTrackingGestureFilterLevel GestureFilterLevel)
+{
+#if WITH_MLSDK
+	ConditionallyEnable();
+
+	if (!MLHandleIsValid(HandTracker))
+	{
+		return false;
+	}
+
+	// -1 because NoPose is not configurable.
+	const uint32 MaxConfigurablePose = MLHandTrackingKeyPose_Count - 2;
+	MLHandTrackingConfiguration Config;
+	for (uint32 i = 0; i <= MaxConfigurablePose; ++i)
+	{
+		Config.keypose_config[i] = false;
+	}
+
+	bool EnableHandTrackingPipeline = false;
+
+	for (const EHandTrackingGesture& StaticPose : ActiveKeyPoses)
+	{
+		if (static_cast<uint32>(StaticPose) <= MaxConfigurablePose)
+		{
+			Config.keypose_config[static_cast<uint32>(StaticPose)] = true;
+		}
+	}
+
+	Config.handtracking_pipeline_enabled = bTrackingEnabled;
+	Config.keypoints_filter_level = UnrealToMLKeypointsFilterLevel(KeypointsFilterLevel);
+	Config.pose_filter_level = UnrealToMLGestureFilterLevel(GestureFilterLevel);
+
+	return MLHandTrackingSetConfiguration(HandTracker, &Config) == MLResult_Ok;
+#else
+	return false;
+#endif //WITH_MLSDK
+}
+
+bool FMagicLeapHandTracking::GetConfiguration(bool& bTrackingEnabled, TArray<EHandTrackingGesture>& ActiveKeyPoses, EHandTrackingKeypointFilterLevel& KeypointsFilterLevel, EHandTrackingGestureFilterLevel& GestureFilterLevel)
+{
+#if WITH_MLSDK
+	ConditionallyEnable();
+
+	if (!MLHandleIsValid(HandTracker))
+	{
+		return false;
+	}
+
+	MLHandTrackingConfiguration config;
+	bool validResult = MLHandTrackingGetConfiguration(HandTracker, &config) == MLResult_Ok;
+
+	if (validResult)
+	{
+		ActiveKeyPoses.Empty(MLHandTrackingKeyPose_Count);
+		for (uint32 i = 0; i < MLHandTrackingKeyPose_Count; ++i)
+		{
+			if (config.keypose_config[i])
+			{
+				ActiveKeyPoses.Add(static_cast<EHandTrackingGesture>(i));
+			}
+		}
+
+		bTrackingEnabled = config.handtracking_pipeline_enabled;
+		KeypointsFilterLevel = MLToUnrealKeypointsFilterLevel(config.keypoints_filter_level);
+		GestureFilterLevel = MLToUnrealGestureFilterLevel(config.pose_filter_level);
+	}
+
+	return validResult;
+#else
+	return false;
+#endif //WITH_MLSDK
+
+}
+
+void FMagicLeapHandTracking::SetGestureConfidenceThreshold(EHandTrackingGesture Gesture, float Confidence)
+{
+	if (Gesture <= EHandTrackingGesture::NoHand)
+	{
+		GestureConfidenceThresholds[static_cast<int32>(Gesture)] = Confidence;
+	}
+}
+
+float FMagicLeapHandTracking::GetGestureConfidenceThreshold(EHandTrackingGesture Gesture) const
+{
+	if (Gesture <= EHandTrackingGesture::NoHand)
+	{
+		return GestureConfidenceThresholds[static_cast<int32>(Gesture)];
+	}
+	return 0.0f;
+}
+
+namespace MagicLeapHandTracking
+{
+#if WITH_MLSDK
+	bool FetchTransform(const FAppFramework& AppFramework, const MLKeyPointState& Source, FTransform& OutDest, const TCHAR* DebugString, const TCHAR* DebugString2, const TCHAR* DebugString3)
+	{
+		bool bResult = false;
+
+		if (Source.is_valid)
+		{
+			EFailReason FailReason = EFailReason::None;
+			bResult = AppFramework.GetTransform(Source.frame_id, OutDest, FailReason);
+			if (!bResult)
+			{
+				if (FailReason == EFailReason::NaNsInTransform)
+				{
+					UE_LOG(LogMagicLeapHandTracking, Error, TEXT("NaNs in %s %s %s transform."), DebugString, DebugString2, DebugString3);
+				}
+			}
+
+		}
+		return bResult;
+	}
+
+	void FetchFingerTransforms(const FAppFramework& AppFramework, const MLFingerState& Source, FMagicLeapHandTracking::FDigitTransforms& OutDest, const TCHAR* DebugString, const TCHAR* DebugString2)
+	{
+		FetchTransform(AppFramework, Source.tip, OutDest.Tip, DebugString, DebugString2, TEXT("Tip"));
+		FetchTransform(AppFramework, Source.dip, OutDest.DIP, DebugString, DebugString2, TEXT("DIP"));
+		FetchTransform(AppFramework, Source.pip, OutDest.PIP, DebugString, DebugString2, TEXT("PIP"));
+		FetchTransform(AppFramework, Source.mcp, OutDest.MCP, DebugString, DebugString2, TEXT("MCP"));
+	}
+
+	void FetchHandTransforms(const FAppFramework& AppFramework, const MLHandTrackingStaticHandState& Source, FMagicLeapHandTracking::FHandState& OutDest, const TCHAR* DebugString)
+	{
+		FetchTransform(AppFramework, Source.thumb.tip, OutDest.Thumb.Tip, DebugString, TEXT("thumb"), TEXT("Tip"));
+		FetchTransform(AppFramework, Source.thumb.ip,  OutDest.Thumb.DIP, DebugString, TEXT("thumb"), TEXT("DIP"));
+		FetchTransform(AppFramework, Source.thumb.mcp, OutDest.Thumb.PIP, DebugString, TEXT("thumb"), TEXT("PIP"));
+		FetchTransform(AppFramework, Source.thumb.cmc, OutDest.Thumb.MCP, DebugString, TEXT("thumb"), TEXT("MCP"));
+
+		FetchFingerTransforms(AppFramework, Source.index,  OutDest.IndexFinger,  DebugString, TEXT("Index"));
+		FetchFingerTransforms(AppFramework, Source.middle, OutDest.MiddleFinger, DebugString, TEXT("Index"));
+		FetchFingerTransforms(AppFramework, Source.ring,   OutDest.RingFinger,   DebugString, TEXT("Index"));
+		FetchFingerTransforms(AppFramework, Source.pinky,  OutDest.PinkyFinger,  DebugString, TEXT("Index"));
+
+		FetchTransform(AppFramework, Source.wrist.center, OutDest.Wrist.Center, DebugString, TEXT("wrist"), TEXT("center"));
+		FetchTransform(AppFramework, Source.wrist.ulnar,  OutDest.Wrist.Ulnar,  DebugString, TEXT("wrist"), TEXT("ulnar"));
+		FetchTransform(AppFramework, Source.wrist.radial, OutDest.Wrist.Radial, DebugString, TEXT("wrist"), TEXT("radial"));
+
+		FetchTransform(AppFramework, Source.hand_center, OutDest.HandCenter, DebugString, TEXT(""), TEXT("center"));
+	}
+
+	//void LogMLHandTackingHandState(const MLHandTrackingHandState& State)
+	//{
+	//	UE_LOG(LogMagicLeapHandTracking, Log, TEXT("    hand_confidence=%0.2f"), State.hand_confidence);
+	//	UE_LOG(LogMagicLeapHandTracking, Log, TEXT("    keypose=%i (%0.2f)"), State.keypose, State.keypose_confidence[State.keypose]);
+	//	UE_LOG(LogMagicLeapHandTracking, Log, TEXT("    hand_center_normalized=%0.2f, %0.2f, %0.2f"), State.hand_center_normalized.x, State.hand_center_normalized.y, State.hand_center_normalized.z);
+	//	static TCHAR Mask[32] = {0};
+	//	for (size_t i = 0; i < MLHandTrackingStaticData_MaxKeyPoints; i++)
+	//	{
+	//		Mask[i] = State.keypoints_mask[i] ? '1' : '0';
+	//	}
+	//	UE_LOG(LogMagicLeapHandTracking, Log, TEXT("    keypoints_mask=%s"), Mask);
+
+	//}
+	//void LogMLHandTrackingData(const MLHandTrackingData& Data)
+	//{
+	//	UE_LOG(LogMagicLeapHandTracking, Log, TEXT("LogMLHandTrackingData"));
+
+	//	UE_LOG(LogMagicLeapHandTracking, Log, TEXT("  LeftHand"));
+	//	LogMLHandTackingHandState(Data.left_hand_state);
+
+	//	//UE_LOG(LogMagicLeapHandTracking, Log, TEXT("  RightHand"));
+	//	//LogMLHandTackingHandState(Data.right_hand_state);
+	//}
+
+	//void LogKeyPoint(const MLKeyPointState& State)
+	//{
+	//	UE_LOG(LogMagicLeapHandTracking, Log, TEXT("    %i %llu:%llu"), State.is_valid, (unsigned long long)State.frame_id.data[0], (unsigned long long)State.frame_id.data[1]);
+	//}
+
+	//void LogKeyPoint(const MLKeyPointState& State, const TCHAR* DebugString)
+	//{
+	//	UE_LOG(LogMagicLeapHandTracking, Log, TEXT("      %s  %i %llu:%llu"), DebugString, State.is_valid, (unsigned long long)State.frame_id.data[0], (unsigned long long)State.frame_id.data[1]);
+	//}
+	//void LogMLFingerState(const MLFingerState& State, const TCHAR* DebugString)
+	//{
+	//	UE_LOG(LogMagicLeapHandTracking, Log, TEXT("    %s"), DebugString);
+	//	LogKeyPoint(State.tip, TEXT("Tip"));
+	//	LogKeyPoint(State.dip, TEXT("DIP"));
+	//	LogKeyPoint(State.pip, TEXT("PIP"));
+	//	LogKeyPoint(State.mcp, TEXT("MCP"));
+
+	//}
+	//void LogMLHandTrackingStaticHandState(const MLHandTrackingStaticHandState& State)
+	//{
+	//	UE_LOG(LogMagicLeapHandTracking, Log, TEXT("    Thumb"));
+	//	LogKeyPoint(State.thumb.tip, TEXT("Tip"));
+	//	LogKeyPoint(State.thumb.ip,  TEXT("DIP"));
+	//	LogKeyPoint(State.thumb.mcp, TEXT("PIP"));
+	//	LogKeyPoint(State.thumb.cmc, TEXT("MCP"));
+
+	//	LogMLFingerState(State.index,  TEXT("Index"));
+	//	LogMLFingerState(State.middle, TEXT("middle"));
+	//	LogMLFingerState(State.ring,   TEXT("ring"));
+	//	LogMLFingerState(State.pinky,  TEXT("pinky"));
+
+	//	UE_LOG(LogMagicLeapHandTracking, Log, TEXT("    Wrist"));
+	//	LogKeyPoint(State.wrist.center, TEXT("center"));
+	//	LogKeyPoint(State.wrist.ulnar,  TEXT("ulnar"));
+	//	LogKeyPoint(State.wrist.radial, TEXT("radial"));
+
+	//	UE_LOG(LogMagicLeapHandTracking, Log, TEXT("    HandCenter"));
+	//	LogKeyPoint(State.hand_center, TEXT("center"));
+	//}
+	//void LogMLHandTrackingStaticData(const MLHandTrackingStaticData& Data)
+	//{
+	//	UE_LOG(LogMagicLeapHandTracking, Log, TEXT("LogMLHandTrackingStaticData"));
+
+	//	UE_LOG(LogMagicLeapHandTracking, Log, TEXT("  LeftHand"));
+	//	for (size_t i = 0; i < MLHandTrackingStaticData_MaxKeyPoints; i++)
+	//	{
+	//		const MLKeyPointState& KeyPoint = Data.left_frame[i];
+	//		LogKeyPoint(KeyPoint);
+	//	}
+	//	LogMLHandTrackingStaticHandState(Data.left);
+
+	//	//UE_LOG(LogMagicLeapHandTracking, Log, TEXT("  RightHand"));
+	//	//for (size_t i = 0; i < MLHandTrackingStaticData_MaxKeyPoints; i++)
+	//	//{
+	//	//	const MLKeyPointState& KeyPoint = Data.right_frame[i];
+	//	//	UE_LOG(LogMagicLeapHandTracking, Log, TEXT("    %i %ll:%ll"), KeyPoint.is_valid, KeyPoint.frame_id.data[0], KeyPoint.frame_id.data[1]);
+	//	//}
+	//	//MLHandTrackingStaticHandState(Data.right);
+	//}
+#endif // WITH_MLSDK
+}
+
+void FMagicLeapHandTracking::UpdateTrackerData()
+{
+#if WITH_MLSDK
+	if (MLHandleIsValid(HandTracker))
+	{
+		CurrentHandTrackingDataIndex = 1 - CurrentHandTrackingDataIndex;
+		bIsHandTrackingStateValid =		MLHandTrackingGetData(HandTracker, &HandTrackingDatas[CurrentHandTrackingDataIndex]) == MLResult_Ok;
+		//MagicLeapHandTracking::LogMLHandTrackingData(GetCurrentHandTrackingData());
+	}
+	else
+	{
+		bIsHandTrackingStateValid = false;
+	}
+
+	if (bIsHandTrackingStateValid && IMagicLeapPlugin::Get().IsMagicLeapHMDValid())
+	{
+		const FAppFramework& AppFramework = static_cast<FMagicLeapHMD*>(GEngine->XRSystem->GetHMDDevice())->GetAppFrameworkConst();
+		check(AppFramework.IsInitialized());
+
+		const MLHandTrackingData& HandTrackingData = GetCurrentHandTrackingData();
+
+		LeftHand.Gesture = TranslateGestureEnum(HandTrackingData.left_hand_state.keypose);
+		if (LeftHand.Gesture != EHandTrackingGesture::NoHand)
+		{
+			LeftHand.HandCenterNormalized = MagicLeap::ToFVector(HandTrackingData.left_hand_state.hand_center_normalized, 1.0f);
+			if (LeftHand.HandCenterNormalized.ContainsNaN())
+			{
+				UE_LOG(LogMagicLeapHandTracking, Error, TEXT("LeftHand.HandCenterNormalized received NaNs from the device. Zeroing out the vector."));
+				LeftHand.HandCenterNormalized = FVector::ZeroVector;
+			}
+
+			MagicLeapHandTracking::FetchHandTransforms(AppFramework, HandTrackingStaticData.left, LeftHand, TEXT("left hand"));
+		}
+
+		RightHand.Gesture = TranslateGestureEnum(HandTrackingData.right_hand_state.keypose);
+		if (RightHand.Gesture != EHandTrackingGesture::NoHand)
+		{
+			RightHand.HandCenterNormalized = MagicLeap::ToFVector(HandTrackingData.right_hand_state.hand_center_normalized, 1.0f);
+			if (RightHand.HandCenterNormalized.ContainsNaN())
+			{
+				UE_LOG(LogMagicLeapHandTracking, Error, TEXT("RightHand.HandCenterNormalized received NaNs from the device. Zeroing out the vector."));
+				RightHand.HandCenterNormalized = FVector::ZeroVector;
+			}
+
+			MagicLeapHandTracking::FetchHandTransforms(AppFramework, HandTrackingStaticData.right, RightHand, TEXT("right hand"));
+		}
+
+		LeftHand.GestureConfidence = HandTrackingData.left_hand_state.keypose_confidence[HandTrackingData.left_hand_state.keypose];
+		RightHand.GestureConfidence = HandTrackingData.right_hand_state.keypose_confidence[HandTrackingData.right_hand_state.keypose];
+	}
+#endif //WITH_MLSDK
+}
+
+void FMagicLeapHandTracking::AddKeys()
+{
+	// Left Static HandTracking
+	EKeys::AddKey(FKeyDetails(FStaticHandTracking::Left_Finger, LOCTEXT("MagicLeap_Left_Finger", "ML Left Finger"), FKeyDetails::GamepadKey));
+	EKeys::AddKey(FKeyDetails(FStaticHandTracking::Left_Fist, LOCTEXT("MagicLeap_Left_Fist", "ML Left Fist"), FKeyDetails::GamepadKey));
+	EKeys::AddKey(FKeyDetails(FStaticHandTracking::Left_Pinch, LOCTEXT("MagicLeap_Left_Pinch", "ML Left Pinch"), FKeyDetails::GamepadKey));
+	EKeys::AddKey(FKeyDetails(FStaticHandTracking::Left_Thumb, LOCTEXT("MagicLeap_Left_Thumb", "ML Left Thumb"), FKeyDetails::GamepadKey));
+	EKeys::AddKey(FKeyDetails(FStaticHandTracking::Left_L, LOCTEXT("MagicLeap_Left_L", "ML Left L"), FKeyDetails::GamepadKey));
+	EKeys::AddKey(FKeyDetails(FStaticHandTracking::Left_OpenHandBack, LOCTEXT("MagicLeap_Left_OpenHandBack", "ML Left Open Hand Back"), FKeyDetails::GamepadKey));
+	EKeys::AddKey(FKeyDetails(FStaticHandTracking::Left_Ok, LOCTEXT("MagicLeap_Left_Ok", "ML Left Ok"), FKeyDetails::GamepadKey));
+	EKeys::AddKey(FKeyDetails(FStaticHandTracking::Left_C, LOCTEXT("MagicLeap_Left_C", "ML Left C"), FKeyDetails::GamepadKey));
+	EKeys::AddKey(FKeyDetails(FStaticHandTracking::Left_NoPose, LOCTEXT("MagicLeap_Left_NoPose", "ML Left NoPose"), FKeyDetails::GamepadKey));
+	EKeys::AddKey(FKeyDetails(FStaticHandTracking::Left_NoHand, LOCTEXT("MagicLeap_Left_NoHand", "ML Left No Hand"), FKeyDetails::GamepadKey));
+
+	// Right Static HandTracking
+	EKeys::AddKey(FKeyDetails(FStaticHandTracking::Right_Finger, LOCTEXT("MagicLeap_Right_Finger", "ML Right Finger"), FKeyDetails::GamepadKey));
+	EKeys::AddKey(FKeyDetails(FStaticHandTracking::Right_Fist, LOCTEXT("MagicLeap_Right_Fist", "ML Right Fist"), FKeyDetails::GamepadKey));
+	EKeys::AddKey(FKeyDetails(FStaticHandTracking::Right_Pinch, LOCTEXT("MagicLeap_Right_Pinch", "ML Right Pinch"), FKeyDetails::GamepadKey));
+	EKeys::AddKey(FKeyDetails(FStaticHandTracking::Right_Thumb, LOCTEXT("MagicLeap_Right_Thumb", "ML Right Thumb"), FKeyDetails::GamepadKey));
+	EKeys::AddKey(FKeyDetails(FStaticHandTracking::Right_L, LOCTEXT("MagicLeap_Right_L", "ML Right L"), FKeyDetails::GamepadKey));
+	EKeys::AddKey(FKeyDetails(FStaticHandTracking::Right_OpenHandBack, LOCTEXT("MagicLeap_Right_OpenHandBack", "ML Right Open Hand Back"), FKeyDetails::GamepadKey));
+	EKeys::AddKey(FKeyDetails(FStaticHandTracking::Right_Ok, LOCTEXT("MagicLeap_Right_Ok", "ML Right Ok"), FKeyDetails::GamepadKey));
+	EKeys::AddKey(FKeyDetails(FStaticHandTracking::Right_C, LOCTEXT("MagicLeap_Right_C", "ML Right C"), FKeyDetails::GamepadKey));
+	EKeys::AddKey(FKeyDetails(FStaticHandTracking::Right_NoPose, LOCTEXT("MagicLeap_Right_NoPose", "ML Right NoPose"), FKeyDetails::GamepadKey));
+	EKeys::AddKey(FKeyDetails(FStaticHandTracking::Right_NoHand, LOCTEXT("MagicLeap_Right_NoHand", "ML Right No Hand"), FKeyDetails::GamepadKey));
+
+	// Create mapping of static and dynamic gesture enums to their key names.
+
+#if WITH_MLSDK
+	LeftStaticGestureMap.AddDefaulted(static_cast<int32>(MLHandTrackingKeyPose_Count));
+
+	RightStaticGestureMap.AddDefaulted(static_cast<int32>(MLHandTrackingKeyPose_Count));
+
+	// Left Static HandTracking
+	LeftStaticGestureMap[MLHandTrackingKeyPose_Finger] = FMagicLeapGestureKeyNames::Left_Finger_Name;
+	LeftStaticGestureMap[MLHandTrackingKeyPose_Fist] = FMagicLeapGestureKeyNames::Left_Fist_Name;
+	LeftStaticGestureMap[MLHandTrackingKeyPose_Pinch] = FMagicLeapGestureKeyNames::Left_Pinch_Name;
+	LeftStaticGestureMap[MLHandTrackingKeyPose_Thumb] = FMagicLeapGestureKeyNames::Left_Thumb_Name;
+	LeftStaticGestureMap[MLHandTrackingKeyPose_L] = FMagicLeapGestureKeyNames::Left_L_Name;
+	LeftStaticGestureMap[MLHandTrackingKeyPose_OpenHandBack] = FMagicLeapGestureKeyNames::Left_OpenHandBack_Name;
+	LeftStaticGestureMap[MLHandTrackingKeyPose_Ok] = FMagicLeapGestureKeyNames::Left_Ok_Name;
+	LeftStaticGestureMap[MLHandTrackingKeyPose_C] = FMagicLeapGestureKeyNames::Left_C_Name;
+	LeftStaticGestureMap[MLHandTrackingKeyPose_NoPose] = FMagicLeapGestureKeyNames::Left_NoPose_Name;
+	LeftStaticGestureMap[MLHandTrackingKeyPose_NoHand] = FMagicLeapGestureKeyNames::Left_NoHand_Name;
+
+	// Right Static HandTracking
+	RightStaticGestureMap[MLHandTrackingKeyPose_Finger] = FMagicLeapGestureKeyNames::Right_Finger_Name;
+	RightStaticGestureMap[MLHandTrackingKeyPose_Fist] = FMagicLeapGestureKeyNames::Right_Fist_Name;
+	RightStaticGestureMap[MLHandTrackingKeyPose_Pinch] = FMagicLeapGestureKeyNames::Right_Pinch_Name;
+	RightStaticGestureMap[MLHandTrackingKeyPose_Thumb] = FMagicLeapGestureKeyNames::Right_Thumb_Name;
+	RightStaticGestureMap[MLHandTrackingKeyPose_L] = FMagicLeapGestureKeyNames::Right_L_Name;
+	RightStaticGestureMap[MLHandTrackingKeyPose_OpenHandBack] = FMagicLeapGestureKeyNames::Right_OpenHandBack_Name;
+	RightStaticGestureMap[MLHandTrackingKeyPose_Ok] = FMagicLeapGestureKeyNames::Right_Ok_Name;
+	RightStaticGestureMap[MLHandTrackingKeyPose_C] = FMagicLeapGestureKeyNames::Right_C_Name;
+	RightStaticGestureMap[MLHandTrackingKeyPose_NoPose] = FMagicLeapGestureKeyNames::Right_NoPose_Name;
+	RightStaticGestureMap[MLHandTrackingKeyPose_NoHand] = FMagicLeapGestureKeyNames::Right_NoHand_Name;
+#endif //WITH_MLSDK
+
+}
+
+#if WITH_MLSDK
+EHandTrackingGesture FMagicLeapHandTracking::TranslateGestureEnum(MLHandTrackingKeyPose HandState)
+{
+	switch (HandState)
+	{
+	case MLHandTrackingKeyPose_Finger:
+		return EHandTrackingGesture::Finger;
+	case MLHandTrackingKeyPose_Fist:
+		return EHandTrackingGesture::Fist;
+	case MLHandTrackingKeyPose_Pinch:
+		return EHandTrackingGesture::Pinch;
+	case MLHandTrackingKeyPose_Thumb:
+		return EHandTrackingGesture::Thumb;
+	case MLHandTrackingKeyPose_L:
+		return EHandTrackingGesture::L;
+	case MLHandTrackingKeyPose_OpenHandBack:
+		return EHandTrackingGesture::OpenHandBack;
+	case MLHandTrackingKeyPose_Ok:
+		return EHandTrackingGesture::Ok;
+	case MLHandTrackingKeyPose_C:
+		return EHandTrackingGesture::C;
+	case MLHandTrackingKeyPose_NoPose:
+		return EHandTrackingGesture::NoPose;
+	case MLHandTrackingKeyPose_NoHand:
+		return EHandTrackingGesture::NoHand;
+	default:
+		check(false);
+		return EHandTrackingGesture::NoHand;
+	}
+}
+#endif //WITH_MLSDK
+
+void FMagicLeapHandTracking::ConditionallyEnable()
+{
+#if WITH_MLSDK
+	if (!MLHandleIsValid(HandTracker) && GEngine->XRSystem.IsValid() && GEngine->XRSystem->GetHMDDevice())
+	{
+		MLResult CreateResult = MLHandTrackingCreate(&HandTracker);
+		if (CreateResult == MLResult_Ok && MLHandleIsValid(HandTracker))
+		{
+			if (MLResult_Ok != MLHandTrackingGetStaticData(HandTracker, &HandTrackingStaticData))
+			{
+				UE_LOG(LogMagicLeapHandTracking, Error, TEXT("Error getting hand tracker static data."));
+			}
+			//MagicLeapHandTracking::LogMLHandTrackingStaticData(HandTrackingStaticData);
+		}
+		else
+		{
+			UE_LOG(LogMagicLeapHandTracking, Error, TEXT("Error creating hand tracker."));
+		}
+	}
+#endif //WITH_MLSDK
+}
+
+void FMagicLeapHandTracking::OnAppPause()
+{
+#if WITH_MLSDK
+	if (!MLHandleIsValid(HandTracker))
+	{
+		UE_LOG(LogMagicLeapHandTracking, Error, TEXT("Hand tracker was invalid on application pause."));
+	}
+	else
+	{
+		MLHandTrackingConfiguration HandTrackingConfig;
+
+		if (MLResult_Ok != MLHandTrackingGetConfiguration(HandTracker, &HandTrackingConfig))
+		{
+			UE_LOG(LogMagicLeapHandTracking, Error, TEXT("Failed to retrieve hand tracking configuration on application pause."));
+		}
+		else
+		{
+			bWasSystemEnabledOnPause = HandTrackingConfig.handtracking_pipeline_enabled;
+
+			if (!bWasSystemEnabledOnPause)
+			{
+				UE_LOG(LogMagicLeapHandTracking, Log, TEXT("Hand tracking was not enabled at time of application pause."));
+			}
+			else
+			{
+				HandTrackingConfig.handtracking_pipeline_enabled = false;
+
+				if (MLResult_Ok != MLHandTrackingSetConfiguration(HandTracker, &HandTrackingConfig))
+				{
+					UE_LOG(LogMagicLeapHandTracking, Error, TEXT("Failed to disable hand tracking on application pause."));
+				}
+				else
+				{
+					UE_LOG(LogMagicLeapHandTracking, Log, TEXT("Hand tracking paused until app resumes."));
+				}
+			}
+		}
+	}
+#endif //WITH_MLSDK
+}
+
+void FMagicLeapHandTracking::OnAppResume()
+{
+#if WITH_MLSDK
+	if (!MLHandleIsValid(HandTracker))
+	{
+		UE_LOG(LogMagicLeapHandTracking, Error, TEXT("Hand tracker was invalid on application resume."));
+	}
+	else
+	{
+		MLHandTrackingConfiguration HandTrackingConfig;
+
+		if (!bWasSystemEnabledOnPause)
+		{
+			UE_LOG(LogMagicLeapHandTracking, Log, TEXT("Not resuming hand tracking as it was not enabled at time of application pause."));
+		}
+		else
+		{
+			if (MLResult_Ok != MLHandTrackingGetConfiguration(HandTracker, &HandTrackingConfig))
+			{
+				UE_LOG(LogMagicLeapHandTracking, Error, TEXT("Failed to retrieve hand tracking configuration on application resume."));
+			}
+			else
+			{
+				HandTrackingConfig.handtracking_pipeline_enabled = true;
+
+				if (MLResult_Ok != MLHandTrackingSetConfiguration(HandTracker, &HandTrackingConfig))
+				{
+					UE_LOG(LogMagicLeapHandTracking, Error, TEXT("Failed to re-enable hand tracking on application resume."));
+				}
+				else
+				{
+					UE_LOG(LogMagicLeapHandTracking, Log, TEXT("Hand tracking resumed."));
+				}
+			}
+		}
+	}
+#endif //WITH_MLSDK
+}
+
+#undef LOCTEXT_NAMESPACE
diff --git a/Engine/Plugins/Lumin/MagicLeap/Source/MagicLeapHandTracking/Private/MagicLeapHandTracking.h b/Engine/Plugins/Lumin/MagicLeap/Source/MagicLeapHandTracking/Private/MagicLeapHandTracking.h
new file mode 100644
--- /dev/null
+++ b/Engine/Plugins/Lumin/MagicLeap/Source/MagicLeapHandTracking/Private/MagicLeapHandTracking.h
@@ -0,0 +1,233 @@
+// Copyright 1998-2018 Epic Games, Inc. All Rights Reserved.
+
+#pragma once
+
+#include "CoreMinimal.h"
+#include "GenericPlatform/IInputInterface.h"
+#include "IMagicLeapInputDevice.h"
+#include "MagicLeapPluginUtil.h" // for ML_INCLUDES_START/END
+
+#if WITH_MLSDK
+ML_INCLUDES_START
+#include <ml_hand_tracking.h>
+ML_INCLUDES_END
+#endif //WITH_MLSDK
+
+#include "MagicLeapHandTrackingTypes.h"
+#include "XRMotionControllerBase.h"
+#include "InputCoreTypes.h"
+#include "AppEventHandler.h"
+
+/**
+  * MagicLeap HandTracking
+  */
+class FMagicLeapHandTracking : public IMagicLeapInputDevice, public FXRMotionControllerBase, public MagicLeap::IAppEventHandler
+{
+public:
+	static FName LeftHandCenter_Name;
+
+	static FName LeftWristCenter_Name;
+	static FName LeftWristUlnar_Name;
+	static FName LeftWristRadial_Name;
+
+	static FName LeftThumbTip_Name;
+	static FName LeftThumbIP_Name;
+	static FName LeftThumbMCP_Name;
+	static FName LeftThumbCMC_Name;
+
+	static FName LeftIndexFingerTip_Name;
+	static FName LeftIndexFingerDIP_Name;
+	static FName LeftIndexFingerPIP_Name;
+	static FName LeftIndexFingerMCP_Name;
+
+	static FName LeftMiddleFingerTip_Name;
+	static FName LeftMiddleFingerDIP_Name;
+	static FName LeftMiddleFingerPIP_Name;
+	static FName LeftMiddleFingerMCP_Name;
+
+	static FName LeftRingFingerTip_Name;
+	static FName LeftRingFingerDIP_Name;
+	static FName LeftRingFingerPIP_Name;
+	static FName LeftRingFingerMCP_Name;
+
+	static FName LeftPinkyFingerTip_Name;
+	static FName LeftPinkyFingerDIP_Name;
+	static FName LeftPinkyFingerPIP_Name;
+	static FName LeftPinkyFingerMCP_Name;
+
+
+	static FName RightHandCenter_Name;
+
+	static FName RightWristCenter_Name;
+	static FName RightWristUlnar_Name;
+	static FName RightWristRadial_Name;
+
+	static FName RightThumbTip_Name;
+	static FName RightThumbIP_Name;
+	static FName RightThumbMCP_Name;
+	static FName RightThumbCMC_Name;
+
+	static FName RightIndexFingerTip_Name;
+	static FName RightIndexFingerDIP_Name;
+	static FName RightIndexFingerPIP_Name;
+	static FName RightIndexFingerMCP_Name;
+
+	static FName RightMiddleFingerTip_Name;
+	static FName RightMiddleFingerDIP_Name;
+	static FName RightMiddleFingerPIP_Name;
+	static FName RightMiddleFingerMCP_Name;
+
+	static FName RightRingFingerTip_Name;
+	static FName RightRingFingerDIP_Name;
+	static FName RightRingFingerPIP_Name;
+	static FName RightRingFingerMCP_Name;
+
+	static FName RightPinkyFingerTip_Name;
+	static FName RightPinkyFingerDIP_Name;
+	static FName RightPinkyFingerPIP_Name;
+	static FName RightPinkyFingerMCP_Name;
+
+	//struct FTransformState
+	//{
+
+	//};
+	struct FWristTransforms : public FNoncopyable
+	{
+		// Wrist center.
+		FTransform Center = FTransform::Identity;
+		// Ulnar-sided wrist
+		FTransform Ulnar = FTransform::Identity;
+		// Radial-sided wrist.
+		FTransform Radial = FTransform::Identity;
+	};
+	struct FDigitTransforms : public FNoncopyable
+	{
+		// These labels are not correct anatomical nomenclature for the thumb, but they 1:1 map with the correct names.
+		FTransform Tip = FTransform::Identity; // digit tip
+		FTransform DIP = FTransform::Identity; // distal joint
+		FTransform PIP = FTransform::Identity; // proximal joint
+		FTransform MCP = FTransform::Identity; // base of digit
+	};
+	struct FHandState : public FNoncopyable
+	{
+		EHandTrackingGesture Gesture = EHandTrackingGesture::NoHand;
+		float GestureConfidence = 0.0f;
+		FVector HandCenterNormalized = FVector::ZeroVector;
+
+		FDigitTransforms Thumb;
+		FDigitTransforms IndexFinger;
+		FDigitTransforms MiddleFinger;
+		FDigitTransforms RingFinger;
+		FDigitTransforms PinkyFinger;
+
+		FWristTransforms Wrist;
+
+		FTransform HandCenter = FTransform::Identity;
+
+		bool IsValid() const { return Gesture != EHandTrackingGesture::NoHand; }
+	};
+
+public:
+	FMagicLeapHandTracking(const TSharedRef<FGenericApplicationMessageHandler>& InMessageHandler);
+	virtual ~FMagicLeapHandTracking();
+
+	/** IMotionController interface */
+	virtual bool GetControllerOrientationAndPosition(const int32 ControllerIndex, const FName MotionSource, FRotator& OutOrientation, FVector& OutPosition, float WorldToMetersScale) const override;
+	virtual bool GetControllerOrientationAndPosition(const int32 ControllerIndex, const EControllerHand DeviceHand, FRotator& OutOrientation, FVector& OutPosition, float WorldToMetersScale) const override;
+	virtual ETrackingStatus GetControllerTrackingStatus(const int32 ControllerIndex, const EControllerHand DeviceHand) const override;
+	virtual FName GetMotionControllerDeviceTypeName() const override;
+	virtual void EnumerateSources(TArray<FMotionControllerSource>& SourcesOut) const override;
+
+
+	/** IMagicLeapInputDevice interface */
+	virtual void Tick(float DeltaTime) override;
+	virtual void SendControllerEvents() override;
+	virtual void SetMessageHandler(const TSharedRef<FGenericApplicationMessageHandler>& InMessageHandler) override;
+	virtual bool Exec(UWorld* InWorld, const TCHAR* Cmd, FOutputDevice& Ar) override;
+	virtual void SetChannelValue(int32 ControllerId, FForceFeedbackChannelType ChannelType, float Value) override {};
+	virtual void SetChannelValues(int32 ControllerId, const FForceFeedbackValues &values) override {};
+	virtual bool IsGamepadAttached() const override;
+	virtual void Enable() override;
+	virtual bool SupportsExplicitEnable() const override;
+	virtual void Disable() override;
+
+	const FHandState& GetLeftHandState() const;
+	const FHandState& GetRightHandState() const;
+	bool IsHandTrackingStateValid() const;
+
+	bool SetConfiguration(bool bTrackingEnabled, const TArray<EHandTrackingGesture>& ActiveKeyPoses, EHandTrackingKeypointFilterLevel KeypointsFilterLevel, EHandTrackingGestureFilterLevel GestureFilterLevel);
+	bool GetConfiguration(bool& bTrackingEnabled, TArray<EHandTrackingGesture>& ActiveKeyPoses, EHandTrackingKeypointFilterLevel& KeypointsFilterLevel, EHandTrackingGestureFilterLevel& GestureFilterLevel);
+
+	void SetGestureConfidenceThreshold(EHandTrackingGesture Gesture, float Confidence);
+	float GetGestureConfidenceThreshold(EHandTrackingGesture Gesture) const;
+
+private:
+	void BuildKeypointMaps();
+	const FTransform* FindTransformBySource(FName SourceName) const;
+	const FHandState* FindHandBySource(FName SourceName) const;
+	void UpdateTrackerData();
+	void AddKeys();
+	void ConditionallyEnable();
+	void OnAppPause() override;
+	void OnAppResume() override;
+
+#if WITH_MLSDK
+	const MLHandTrackingData& GetCurrentHandTrackingData()				{ return HandTrackingDatas[CurrentHandTrackingDataIndex]; }
+	const MLHandTrackingData& GetPreviousHandTrackingData()				{ return HandTrackingDatas[1 - CurrentHandTrackingDataIndex]; }
+	void SendControllerEventsForHand(const MLHandTrackingHandState& NewHandState, const MLHandTrackingHandState& OldHandState, const TArray<FName>& GestureMap);
+
+	static EHandTrackingGesture TranslateGestureEnum(MLHandTrackingKeyPose KeyPose);
+#endif //WITH_MLSDK
+
+	TSharedPtr<FGenericApplicationMessageHandler> MessageHandler;
+	int32 DeviceIndex;
+
+	bool HandTrackingPipelineEnabled;
+
+#if WITH_MLSDK
+	MLHandle HandTracker;
+	MLHandTrackingData HandTrackingDatas[2];
+	int32 CurrentHandTrackingDataIndex = 0;
+	MLHandTrackingStaticData HandTrackingStaticData;
+#endif //WITH_MLSDK
+
+	FHandState LeftHand;
+	FHandState RightHand;
+
+	bool bIsHandTrackingStateValid;
+
+	TArray<float> GestureConfidenceThresholds;
+
+	TArray<FName> LeftStaticGestureMap;
+	TArray<FName> RightStaticGestureMap;
+
+	TMap<FName, FTransform*> SourceToTransformMap;
+	TMap<FName, FHandState*> SourceToHandMap;
+
+	struct FStaticHandTracking
+	{
+		static const FKey Left_Finger;
+		static const FKey Left_Fist;
+		static const FKey Left_Pinch;
+		static const FKey Left_Thumb;
+		static const FKey Left_L;
+		static const FKey Left_OpenHandBack;
+		static const FKey Left_Ok;
+		static const FKey Left_C;
+		static const FKey Left_NoPose;
+		static const FKey Left_NoHand;
+
+		static const FKey Right_Finger;
+		static const FKey Right_Fist;
+		static const FKey Right_Pinch;
+		static const FKey Right_Thumb;
+		static const FKey Right_L;
+		static const FKey Right_OpenHandBack;
+		static const FKey Right_Ok;
+		static const FKey Right_C;
+		static const FKey Right_NoPose;
+		static const FKey Right_NoHand;
+	};
+};
+
+DEFINE_LOG_CATEGORY_STATIC(LogMagicLeapHandTracking, Display, All);
diff --git a/Engine/Plugins/Lumin/MagicLeap/Source/MagicLeapHandTracking/Private/MagicLeapHandTrackingFunctionLibrary.cpp b/Engine/Plugins/Lumin/MagicLeap/Source/MagicLeapHandTracking/Private/MagicLeapHandTrackingFunctionLibrary.cpp
new file mode 100644
--- /dev/null
+++ b/Engine/Plugins/Lumin/MagicLeap/Source/MagicLeapHandTracking/Private/MagicLeapHandTrackingFunctionLibrary.cpp
@@ -0,0 +1,254 @@
+// Copyright 1998-2018 Epic Games, Inc. All Rights Reserved.
+
+#include "MagicLeapHandTrackingFunctionLibrary.h"
+#include "IMagicLeapHandTrackingPlugin.h"
+#include "MagicLeapHandTracking.h"
+#include "IMagicLeapPlugin.h"
+#include "HeadMountedDisplayFunctionLibrary.h"
+#include "Engine/World.h"
+
+bool UMagicLeapHandTrackingFunctionLibrary::GetHandCenter(EControllerHand Hand, FTransform& HandCenter)
+{
+	TSharedPtr<FMagicLeapHandTracking> HandTracking = StaticCastSharedPtr<FMagicLeapHandTracking>(IMagicLeapHandTrackingPlugin::Get().GetInputDevice());
+
+	if (HandTracking.IsValid() && HandTracking->IsHandTrackingStateValid())
+	{
+		if (Hand == EControllerHand::Left || Hand == EControllerHand::Right)
+		{
+			const FMagicLeapHandTracking::FHandState& HandTrackingData = (Hand == EControllerHand::Right) ? HandTracking->GetRightHandState() : HandTracking->GetLeftHandState();
+			HandCenter = HandTrackingData.HandCenter;
+			return HandTrackingData.IsValid();
+		}
+		else
+		{
+			UE_LOG(LogMagicLeapHandTracking, Error, TEXT("Hand %d is not supported"), static_cast<int32>(Hand));
+			return false;
+		}
+	}
+
+	return false;
+}
+
+bool UMagicLeapHandTrackingFunctionLibrary::GetHandIndexFingerTip(EControllerHand Hand, EGestureTransformSpace TransformSpace, FTransform& Pointer)
+{
+	TSharedPtr<FMagicLeapHandTracking> HandTracking = StaticCastSharedPtr<FMagicLeapHandTracking>(IMagicLeapHandTrackingPlugin::Get().GetInputDevice());
+
+	if (HandTracking.IsValid() && HandTracking->IsHandTrackingStateValid())
+	{
+		if (Hand == EControllerHand::Left || Hand == EControllerHand::Right)
+		{
+			const FMagicLeapHandTracking::FHandState& HandTrackingData = (Hand == EControllerHand::Right) ? HandTracking->GetRightHandState() : HandTracking->GetLeftHandState();
+			switch (TransformSpace)
+			{
+			case EGestureTransformSpace::Tracking:
+			{
+				Pointer = HandTrackingData.IndexFinger.Tip;
+				break;
+			}
+			case EGestureTransformSpace::World:
+			{
+				FTransform TrackingToWorldTransform = UHeadMountedDisplayFunctionLibrary::GetTrackingToWorldTransform(GWorld);
+				Pointer = HandTrackingData.IndexFinger.Tip * TrackingToWorldTransform;
+				break;
+			}
+			case EGestureTransformSpace::Hand:
+			{
+				Pointer = HandTrackingData.IndexFinger.Tip * HandTrackingData.HandCenter.Inverse();
+				break;
+			}
+			default:
+				check(false);
+				return false;
+			}
+			return HandTrackingData.IsValid();
+		}
+		else
+		{
+			UE_LOG(LogMagicLeapHandTracking, Error, TEXT("Hand %d is not supported"), static_cast<int32>(Hand));
+			return false;
+		}
+	}
+
+	return false;
+}
+
+bool UMagicLeapHandTrackingFunctionLibrary::GetHandThumbTip(EControllerHand Hand, EGestureTransformSpace TransformSpace, FTransform& Secondary)
+{
+	TSharedPtr<FMagicLeapHandTracking> HandTracking = StaticCastSharedPtr<FMagicLeapHandTracking>(IMagicLeapHandTrackingPlugin::Get().GetInputDevice());
+
+	if (HandTracking.IsValid() && HandTracking->IsHandTrackingStateValid())
+	{
+		if (Hand == EControllerHand::Left || Hand == EControllerHand::Right)
+		{
+			const FMagicLeapHandTracking::FHandState& HandTrackingData = (Hand == EControllerHand::Right) ? HandTracking->GetRightHandState() : HandTracking->GetLeftHandState();
+			switch (TransformSpace)
+			{
+			case EGestureTransformSpace::Tracking:
+			{
+				Secondary = HandTrackingData.Thumb.Tip;
+				break;
+			}
+			case EGestureTransformSpace::World:
+			{
+				FTransform TrackingToWorldTransform = UHeadMountedDisplayFunctionLibrary::GetTrackingToWorldTransform(GWorld);
+				Secondary = HandTrackingData.Thumb.Tip * TrackingToWorldTransform;
+				break;
+			}
+			case EGestureTransformSpace::Hand:
+			{
+				Secondary = HandTrackingData.Thumb.Tip * HandTrackingData.HandCenter.Inverse();
+				break;
+			}
+			default:
+				check(false);
+				return false;
+			}
+			return HandTrackingData.IsValid();
+		}
+		else
+		{
+			UE_LOG(LogMagicLeapHandTracking, Error, TEXT("Hand %d is not supported"), static_cast<int32>(Hand));
+			return false;
+		}
+	}
+
+	return false;
+}
+
+bool UMagicLeapHandTrackingFunctionLibrary::GetHandCenterNormalized(EControllerHand Hand, FVector& HandCenterNormalized)
+{
+	TSharedPtr<FMagicLeapHandTracking> HandTracking = StaticCastSharedPtr<FMagicLeapHandTracking>(IMagicLeapHandTrackingPlugin::Get().GetInputDevice());
+
+	if (HandTracking.IsValid() && HandTracking->IsHandTrackingStateValid())
+	{
+		if (Hand == EControllerHand::Left)
+		{
+			if (HandTracking->GetLeftHandState().IsValid())
+			{
+				HandCenterNormalized = HandTracking->GetLeftHandState().HandCenterNormalized;
+			}
+			return HandTracking->GetLeftHandState().IsValid();
+		}
+		else if (Hand == EControllerHand::Right)
+		{
+			if (HandTracking->GetRightHandState().IsValid())
+			{
+				HandCenterNormalized = HandTracking->GetRightHandState().HandCenterNormalized;
+			}
+			return HandTracking->GetRightHandState().IsValid();
+		}
+		else
+		{
+			UE_LOG(LogMagicLeapHandTracking, Error, TEXT("Hand %d is not supported"), static_cast<int32>(Hand));
+			return false;
+		}
+	}
+
+	return false;
+}
+
+bool UMagicLeapHandTrackingFunctionLibrary::GetGestureKeypoints(EControllerHand Hand, TArray<FTransform>& Keypoints)
+{
+	TSharedPtr<FMagicLeapHandTracking> HandTracking = StaticCastSharedPtr<FMagicLeapHandTracking>(IMagicLeapHandTrackingPlugin::Get().GetInputDevice());
+
+	if (HandTracking.IsValid() && HandTracking->IsHandTrackingStateValid())
+	{
+		if (Hand == EControllerHand::Left || Hand == EControllerHand::Right)
+		{
+			const FMagicLeapHandTracking::FHandState& HandState = (Hand == EControllerHand::Left) ? HandTracking->GetLeftHandState() : HandTracking->GetRightHandState();
+			Keypoints.SetNum(3);
+			Keypoints[0] = HandState.HandCenter;
+			Keypoints[1] = HandState.IndexFinger.Tip;
+			Keypoints[2] = HandState.Thumb.Tip;
+			return true;
+		}
+		else
+		{
+			UE_LOG(LogMagicLeapHandTracking, Error, TEXT("Hand %d is not supported"), static_cast<int32>(Hand));
+			return false;
+		}
+	}
+
+	return false;
+}
+
+bool UMagicLeapHandTrackingFunctionLibrary::SetConfiguration(const TArray<EHandTrackingGesture>& StaticGesturesToActivate, EHandTrackingKeypointFilterLevel KeypointsFilterLevel, EHandTrackingGestureFilterLevel GestureFilterLevel, EHandTrackingGestureFilterLevel HandSwitchingFilterLevel, bool bEnabled)
+{
+	TSharedPtr<FMagicLeapHandTracking> HandTracking = StaticCastSharedPtr<FMagicLeapHandTracking>(IMagicLeapHandTrackingPlugin::Get().GetInputDevice());
+	return HandTracking.IsValid() && HandTracking->SetConfiguration(bEnabled, StaticGesturesToActivate, KeypointsFilterLevel, GestureFilterLevel);
+}
+
+bool UMagicLeapHandTrackingFunctionLibrary::GetConfiguration(TArray<EHandTrackingGesture>& ActiveStaticGestures, EHandTrackingKeypointFilterLevel& KeypointsFilterLevel, EHandTrackingGestureFilterLevel& GestureFilterLevel, EHandTrackingGestureFilterLevel& HandSwitchingFilterLevel, bool& bEnabled)
+{
+	TSharedPtr<FMagicLeapHandTracking> HandTracking = StaticCastSharedPtr<FMagicLeapHandTracking>(IMagicLeapHandTrackingPlugin::Get().GetInputDevice());
+	return HandTracking.IsValid() && HandTracking->GetConfiguration(bEnabled, ActiveStaticGestures, KeypointsFilterLevel, GestureFilterLevel);
+}
+
+void UMagicLeapHandTrackingFunctionLibrary::SetStaticGestureConfidenceThreshold(EHandTrackingGesture Gesture, float Confidence)
+{
+	TSharedPtr<FMagicLeapHandTracking> HandTracking = StaticCastSharedPtr<FMagicLeapHandTracking>(IMagicLeapHandTrackingPlugin::Get().GetInputDevice());
+	if (HandTracking.IsValid())
+	{
+		HandTracking->SetGestureConfidenceThreshold(Gesture, Confidence);
+	}
+}
+
+float UMagicLeapHandTrackingFunctionLibrary::GetStaticGestureConfidenceThreshold(EHandTrackingGesture Gesture)
+{
+	TSharedPtr<FMagicLeapHandTracking> HandTracking = StaticCastSharedPtr<FMagicLeapHandTracking>(IMagicLeapHandTrackingPlugin::Get().GetInputDevice());
+	return (HandTracking.IsValid()) ? HandTracking->GetGestureConfidenceThreshold(Gesture) : 0.0f;
+}
+
+bool UMagicLeapHandTrackingFunctionLibrary::GetCurrentGestureConfidence(EControllerHand Hand, float& Confidence)
+{
+	TSharedPtr<FMagicLeapHandTracking> HandTracking = StaticCastSharedPtr<FMagicLeapHandTracking>(IMagicLeapHandTrackingPlugin::Get().GetInputDevice());
+
+	if (HandTracking.IsValid() && HandTracking->IsHandTrackingStateValid())
+	{
+		if (Hand == EControllerHand::Left)
+		{
+			Confidence = HandTracking->GetLeftHandState().GestureConfidence;
+			return true;
+		}
+		else if (Hand == EControllerHand::Right)
+		{
+			Confidence = HandTracking->GetRightHandState().GestureConfidence;
+			return true;
+		}
+		else
+		{
+			UE_LOG(LogMagicLeapHandTracking, Error, TEXT("Hand %d is not supported"), static_cast<int32>(Hand));
+			return false;
+		}
+	}
+
+	return false;
+}
+
+bool UMagicLeapHandTrackingFunctionLibrary::GetCurrentGesture(EControllerHand Hand, EHandTrackingGesture& Gesture)
+{
+	TSharedPtr<FMagicLeapHandTracking> HandTracking = StaticCastSharedPtr<FMagicLeapHandTracking>(IMagicLeapHandTrackingPlugin::Get().GetInputDevice());
+
+	if (HandTracking.IsValid() && HandTracking->IsHandTrackingStateValid())
+	{
+		if (Hand == EControllerHand::Left)
+		{
+			Gesture = HandTracking->GetLeftHandState().Gesture;
+			return true;
+		}
+		else if (Hand == EControllerHand::Right)
+		{
+			Gesture = HandTracking->GetRightHandState().Gesture;
+			return true;
+		}
+		else
+		{
+			UE_LOG(LogMagicLeapHandTracking, Error, TEXT("Hand %d is not supported"), static_cast<int32>(Hand));
+			Gesture = EHandTrackingGesture::NoHand;
+			return false;
+		}
+	}
+
+	Gesture = EHandTrackingGesture::NoHand;
+	return false;
+}
diff --git a/Engine/Plugins/Lumin/MagicLeap/Source/MagicLeapHandTracking/Public/IMagicLeapHandTrackingPlugin.h b/Engine/Plugins/Lumin/MagicLeap/Source/MagicLeapHandTracking/Public/IMagicLeapHandTrackingPlugin.h
new file mode 100644
--- /dev/null
+++ b/Engine/Plugins/Lumin/MagicLeap/Source/MagicLeapHandTracking/Public/IMagicLeapHandTrackingPlugin.h
@@ -0,0 +1,43 @@
+// Copyright 1998-2018 Epic Games, Inc. All Rights Reserved.
+
+#pragma once
+
+#include "Modules/ModuleManager.h"
+#include "IInputDeviceModule.h"
+
+/**
+ * The public interface to this module.  In most cases, this interface is only public to sibling modules
+ * within this plugin.
+ */
+class IMagicLeapHandTrackingPlugin : public IInputDeviceModule
+{
+public:
+
+	/**
+	 * Singleton-like access to this module's interface.  This is just for convenience!
+	 * Beware of calling this during the shutdown phase, though.  Your module might have been unloaded already.
+	 *
+	 * @return Returns singleton instance, loading the module on demand if needed
+	 */
+	static inline IMagicLeapHandTrackingPlugin& Get()
+	{
+		return FModuleManager::LoadModuleChecked<IMagicLeapHandTrackingPlugin>("MagicLeapHandTracking");
+	}
+
+	/**
+	 * Checks to see if this module is loaded and ready.  It is only valid to call Get() if IsAvailable() returns true.
+	 *
+	 * @return True if the module is loaded and ready to use
+	 */
+	static inline bool IsAvailable()
+	{
+		return FModuleManager::Get().IsModuleLoaded("MagicLeapHandTracking");
+	}
+
+	/**
+	 * Returns the Input Device associated with this InputDeviceModule.
+	 *
+	 * @return Shared pointer to the Input Device.
+	 */
+	virtual TSharedPtr<IInputDevice> GetInputDevice() = 0;
+};
diff --git a/Engine/Plugins/Lumin/MagicLeap/Source/MagicLeapHandTracking/Public/MagicLeapHandTrackingFunctionLibrary.h b/Engine/Plugins/Lumin/MagicLeap/Source/MagicLeapHandTracking/Public/MagicLeapHandTrackingFunctionLibrary.h
new file mode 100644
--- /dev/null
+++ b/Engine/Plugins/Lumin/MagicLeap/Source/MagicLeapHandTracking/Public/MagicLeapHandTrackingFunctionLibrary.h
@@ -0,0 +1,148 @@
+// Copyright 1998-2018 Epic Games, Inc. All Rights Reserved.
+
+#pragma once
+
+#include "MagicLeapHandTrackingTypes.h"
+#include "CoreMinimal.h"
+#include "UObject/ObjectMacros.h"
+#include "Kismet/BlueprintFunctionLibrary.h"
+#include "InputCoreTypes.h"
+#include "MagicLeapHandTrackingFunctionLibrary.generated.h"
+
+UCLASS(ClassGroup = MagicLeap)
+class MAGICLEAPHANDTRACKING_API UMagicLeapHandTrackingFunctionLibrary : public UBlueprintFunctionLibrary
+{
+	GENERATED_BODY()
+
+public:
+
+	/**
+	  Transform of the center of the hand.  Approximately the center of the palm.
+
+	  Note that this returns a transform in the Tracking space. To get the transform in Unreal's
+	  world space, use the MotioController component as a child of the XRPawn with hand set to EControllerHand::Special_1
+	  for the left hand center and EControllerHand::Special_2 for the right hand center.
+
+	  @param Hand Hand to query the hand center transform for. Only Left and Right hand are supported.
+	  @param HandCenter Output parameter containing the position and orientation of the given hand.
+	  @return true if the output param was populated with a valid value, false otherwise.
+	*/
+	UFUNCTION(BlueprintCallable, Category = "HandTracking|MagicLeap")
+	static bool GetHandCenter(EControllerHand Hand, FTransform& HandCenter);
+
+	/**
+	  Transform of the Index Finger Tip.
+
+	  @param Hand Hand to query the hand center transform for. Only Left and Right hand are supported.
+	  @param TransformSpace Get the transform relative to the world, tracking space, or the hand center.
+	  @param Pointer Output parameter containing the position and orientation.
+	  @return true if the output param was populated with a valid value, false otherwise.
+	*/
+	UFUNCTION(BlueprintCallable, Category = "HandTracking|MagicLeap")
+	static bool GetHandIndexFingerTip(EControllerHand Hand, EGestureTransformSpace TransformSpace, FTransform& Pointer);
+
+	/**
+	  Transform of Thumb Tip.
+
+	  Note that this returns a transform in the Tracking space. To get the transform in Unreal's
+	  world space, use the MotioController component as a child of the XRPawn with hand set to EControllerHand::Special_5
+	  for the left hand secondary and EControllerHand::Special_6 for the right hand secondary.
+
+	  @param Hand Hand to query the hand center transform for. Only Left and Right hand are supported.
+	  @param TransformSpace Get the transform relative to the world, tracking space, or the hand center.
+	  @param Secondary Output parameter containing the position and orientation.
+	  @return true if the output param was populated with a valid value, false otherwise.
+	*/
+	UFUNCTION(BlueprintCallable, Category = "HandTracking|MagicLeap")
+	static bool GetHandThumbTip(EControllerHand Hand, EGestureTransformSpace TransformSpace, FTransform& Secondary);
+
+	/**
+	  Normalized position of the center of the given hand. This can be used to detect and warn the users that the hand is out of the gesture detection frame.
+
+	  @param Hand Hand to query the normalized hand center position for. Only Left and Right hand are supported.
+	  @param HandCenterNormalized Output paramter containing the normalized position of the given hand.
+	  @return true if the output param was populated with a valid value, false otherwise.
+	*/
+	UFUNCTION(BlueprintCallable, Category = "HandTracking|MagicLeap")
+	static bool GetHandCenterNormalized(EControllerHand Hand, FVector& HandCenterNormalized);
+
+	/**
+	  List of keypoints detected on the given hand.
+
+	  Note that this returns a transform in the Tracking space. To get the transform in Unreal's
+	  world space, use the MotioController component as a child of the XRPawn with hand set to the following.
+	  Special_3 - Left keypoint 0
+	  Special_5 - Left keypoint 1
+	  Special_4 - Right keypoint 0
+	  Special_6 - Right keypoint 1
+
+	  @param Hand Hand to query the keypoints for. Only Left and Right hand are supported.
+	  @param Keypoints Output paramter containing transforms of the keypoints detected on the given hand.
+	  @return true if the output param was populated with a valid value, false otherwise.
+	*/
+	UFUNCTION(BlueprintCallable, Category = "HandTracking|MagicLeap")
+	static bool GetGestureKeypoints(EControllerHand Hand, TArray<FTransform>& Keypoints);
+
+	/**
+	  Enables and disables the gestures to be detected by the gesture recognition system.
+
+	  @param StaticHandTrackingToActivate List of static gestures to be detected by the system.
+	  @param KeypointsFilterLevel Filtering for the keypoints and hand centers.
+	  @param GestureFilterLevel Filtering for the static gesture recognition.
+	  @param HandSwitchingFilterLevel Filtering for if the left or right hand is present.
+	  @return true if the configuration was set successfully.
+	*/
+	UFUNCTION(BlueprintCallable, Category = "HandTracking|MagicLeap")
+	static bool SetConfiguration(const TArray<EHandTrackingGesture>& StaticGesturesToActivate, EHandTrackingKeypointFilterLevel KeypointsFilterLevel = EHandTrackingKeypointFilterLevel::NoFilter, EHandTrackingGestureFilterLevel GestureFilterLevel = EHandTrackingGestureFilterLevel::NoFilter, EHandTrackingGestureFilterLevel HandSwitchingFilterLevel = EHandTrackingGestureFilterLevel::NoFilter, bool bTrackingEnabled = true);
+
+	/**
+	  Gets the list of static and dynamic gestures currently set to be identified by the gesture recognition system.
+
+	  @param StaticHandTrackingToActivate Output parameter to list the static gestures that can be detected by the system.
+	  @param KeypointsFilterLevel Filtering for the keypoints and hand centers.
+	  @param GestureFilterLevel Filtering for the static gesture recognition.
+	  @param HandSwitchingFilterLevel Filtering for if the left or right hand is present.
+	  @return true if the output params were populated with a valid value, false otherwise.
+	*/
+	UFUNCTION(BlueprintCallable, Category = "HandTracking|MagicLeap")
+	static bool GetConfiguration(TArray<EHandTrackingGesture>& ActiveStaticGestures, EHandTrackingKeypointFilterLevel& KeypointsFilterLevel, EHandTrackingGestureFilterLevel& GestureFilterLevel, EHandTrackingGestureFilterLevel& HandSwitchingFilterLevel, bool& bTrackingEnabled);
+
+	/**
+	  Sets the minimum gesture confidence to filter out the detected static gesture.
+
+	  @param Gesture The gesture to set the confidence threshold for.
+	  @param Confidence The gesture confidence threshold.
+	*/
+	UFUNCTION(BlueprintCallable, Category = "HandTracking|MagicLeap")
+	static void SetStaticGestureConfidenceThreshold(EHandTrackingGesture Gesture, float Confidence);
+
+	/**
+	  Gets the minimum gesture confidence used to filter out the detected static gesture.
+
+	  @param Gesture The gesture to get the confidence threshold for.
+	  @return The gesture confidence threshold.
+	*/
+	UFUNCTION(BlueprintCallable, Category = "HandTracking|MagicLeap")
+	static float GetStaticGestureConfidenceThreshold(EHandTrackingGesture Gesture);
+
+	/**
+	  The confidence level of the current gesture being performed by the given hand.
+	  Value is between [0, 1], 0 is low, 1 is high degree of confidence. For a NoHand, the confidence is always set to 1.
+
+	  @param Hand Hand to query the gesture confidence value for. Only Left and Right hand are supported.
+	  @param Confidence Output parameter containing the confidence value for the given hand's gesture.
+	  @return true if the output param was populated with a valid value, false otherwise.
+	*/
+	UFUNCTION(BlueprintCallable, Category = "HandTracking|MagicLeap")
+	static bool GetCurrentGestureConfidence(EControllerHand Hand, float& Confidence);
+	
+	/**
+	  The current gesture being performed by the given hand.
+
+	  @param Hand Hand to query the gesture for. Only Left and Right hand are supported.
+	  @param Gesture Output parameter containing the given hand's gesture, or NoHand if there isn't one or the system isnt working now.
+	  @return true if the output param was populated with a valid value, false otherwise.
+	*/
+	UFUNCTION(BlueprintCallable, Category = "HandTracking|MagicLeap")
+	static bool GetCurrentGesture(EControllerHand Hand, EHandTrackingGesture& Gesture);
+};
diff --git a/Engine/Plugins/Lumin/MagicLeap/Source/MagicLeapHandTracking/Public/MagicLeapHandTrackingTypes.h b/Engine/Plugins/Lumin/MagicLeap/Source/MagicLeapHandTracking/Public/MagicLeapHandTrackingTypes.h
new file mode 100644
--- /dev/null
+++ b/Engine/Plugins/Lumin/MagicLeap/Source/MagicLeapHandTracking/Public/MagicLeapHandTrackingTypes.h
@@ -0,0 +1,95 @@
+// Copyright 1998-2018 Epic Games, Inc. All Rights Reserved.
+
+#pragma once
+
+#include "CoreMinimal.h"
+#include "UObject/ObjectMacros.h"
+#include "MagicLeapHandTrackingTypes.generated.h"
+
+/*! Static key pose types which are available when both hands are separated. */
+UENUM(BlueprintType)
+enum class EHandTrackingGesture : uint8
+{
+	/** One finger. */
+	Finger,
+	/** A closed fist. */
+	Fist,
+	/** A pinch. */
+	Pinch,
+	/** A closed fist with the thumb pointed up. */
+	Thumb,
+	/** An L shape. */
+	L,
+	/** An open hand with the back of the hand facing the user. */
+	OpenHandBack,
+	/** A pinch with all fingers, except the index finger and the thumb, extended out. */
+	Ok,
+	/** A rounded 'C' alphabet shape. */
+	C,
+	/** No pose detected. */
+	NoPose,
+	/** No hand was present. */
+	NoHand
+};
+
+/** Filtering for the keypoints and hand centers. */
+UENUM(BlueprintType)
+enum class EHandTrackingKeypointFilterLevel : uint8
+{
+	/** No filtering is done, the points are raw. */
+	NoFilter,
+	/** Some smoothing at the cost of latency. */
+	SimpleSmoothing,
+	/** Predictive smoothing, at higher cost of latency. */
+	PredictiveSmoothing
+};
+
+/** Filtering for the gesture recognition and hand switching. */
+UENUM(BlueprintType)
+enum class EHandTrackingGestureFilterLevel : uint8
+{
+	/** No filtering is done, the gestures are raw. */
+	NoFilter,
+	/** Some robustness to flicker at some cost of latency. */
+	SlightRobustnessToFlicker,
+	/** More robust to flicker at higher latency cost. */
+	MoreRobustnessToFlicker
+};
+
+/** Gesture key point transform spaces. */
+UENUM(BlueprintType)
+enum class EGestureTransformSpace : uint8
+{
+	/** Gesture key point transforms are reported in device Tracking space. */
+	Tracking,
+	/** Gesture key point transforms are reported in Unreal world space. This is more costly on CPU.*/
+	World,
+	/** Gesture key point transforms are reported in gesture hand center space.*/
+	Hand
+};
+
+/** List of input key names for all left and right hand gestures. */
+struct FMagicLeapGestureKeyNames
+{
+	static const FName Left_Finger_Name;
+	static const FName Left_Fist_Name;
+	static const FName Left_Pinch_Name;
+	static const FName Left_Thumb_Name;
+	static const FName Left_L_Name;
+	static const FName Left_OpenHandBack_Name;
+	static const FName Left_Ok_Name;
+	static const FName Left_C_Name;
+	static const FName Left_NoPose_Name;
+	static const FName Left_NoHand_Name;
+
+	static const FName Right_Finger_Name;
+	static const FName Right_Fist_Name;
+	static const FName Right_Pinch_Name;
+	static const FName Right_Thumb_Name;
+	static const FName Right_L_Name;
+	static const FName Right_OpenHandBack_Name;
+	static const FName Right_Ok_Name;
+	static const FName Right_C_Name;
+	static const FName Right_NoPose_Name;
+	static const FName Right_NoHand_Name;
+};
