MagicLeapHandTracking CIS fix.

diff --git a/Engine/Plugins/Lumin/MagicLeap/Source/MagicLeapHandTracking/Private/MagicLeapHandTracking.cpp b/Engine/Plugins/Lumin/MagicLeap/Source/MagicLeapHandTracking/Private/MagicLeapHandTracking.cpp
--- a/Engine/Plugins/Lumin/MagicLeap/Source/MagicLeapHandTracking/Private/MagicLeapHandTracking.cpp
+++ b/Engine/Plugins/Lumin/MagicLeap/Source/MagicLeapHandTracking/Private/MagicLeapHandTracking.cpp
@@ -75,9 +75,11 @@ IMPLEMENT_MODULE(FMagicLeapHandTrackingPlugin, MagicLeapHandTracking);
 //////////////////////////////////////////////////////////////////////////
 
 // Partial sanity check of UE4 enums vs ML enums.  If these fail there has probably been an MLSDK update that requires an UE4 update.
+#if WITH_MLSDK
 static_assert(static_cast<int32>(EHandTrackingGesture::NoHand) == static_cast<int32>(MLHandTrackingKeyPose::MLHandTrackingKeyPose_NoHand), "EHandTrackingGesture does not match api enum.");
 static_assert(static_cast<int32>(EHandTrackingKeypointFilterLevel::PredictiveSmoothing) == static_cast<int32>(MLKeypointFilterLevel::MLKeypointFilterLevel_2), "EHandTrackingKeypointFilterLevel does not match api enum.");
 static_assert(static_cast<int32>(EHandTrackingGestureFilterLevel::MoreRobustnessToFlicker) == static_cast<int32>(MLPoseFilterLevel::MLPoseFilterLevel_2), "EHandTrackingGestureFilterLevel does not match api enum.");
+#endif
 
 // Left Gestures
 const FKey FMagicLeapHandTracking::FStaticHandTracking::Left_Finger("MagicLeap_Left_Finger");
