Fixed a missing #ifdef wrapper in the MagicLeap plugin

diff --git a/Engine/Plugins/Lumin/MagicLeap/Source/MagicLeapHelperVulkan/Private/MagicLeapHelperVulkan.cpp b/Engine/Plugins/Lumin/MagicLeap/Source/MagicLeapHelperVulkan/Private/MagicLeapHelperVulkan.cpp
--- a/Engine/Plugins/Lumin/MagicLeap/Source/MagicLeapHelperVulkan/Private/MagicLeapHelperVulkan.cpp
+++ b/Engine/Plugins/Lumin/MagicLeap/Source/MagicLeapHelperVulkan/Private/MagicLeapHelperVulkan.cpp
@@ -1,4 +1,4 @@
-﻿// Copyright 1998-2018 Epic Games, Inc. All Rights Reserved.
+// Copyright 1998-2018 Epic Games, Inc. All Rights Reserved.
 
 #include "MagicLeapHelperVulkan.h"
 #include "IMagicLeapHelperVulkanPlugin.h"
@@ -281,6 +281,8 @@ bool FMagicLeapHelperVulkan::GetMediaTexture(FTextureRHIRef& ResultTexture, FSam
 
 void FMagicLeapHelperVulkan::AliasMediaTexture(FTextureRHIParamRef DestTexture, FTextureRHIParamRef SrcTexture)
 {
+#if PLATFORM_LUMIN
 	FVulkanDynamicRHI* const RHI = (FVulkanDynamicRHI*)GDynamicRHI;
 	RHI->RHIAliasTextureResources(DestTexture, SrcTexture);
-}
\ No newline at end of file
+#endif
+}
