Fixed another missing #ifdef wrapper in the MagicLeap plugin

diff --git a/Engine/Plugins/Lumin/MagicLeap/Source/MagicLeap/Private/MagicLeapHMD.cpp b/Engine/Plugins/Lumin/MagicLeap/Source/MagicLeap/Private/MagicLeapHMD.cpp
--- a/Engine/Plugins/Lumin/MagicLeap/Source/MagicLeap/Private/MagicLeapHMD.cpp
+++ b/Engine/Plugins/Lumin/MagicLeap/Source/MagicLeap/Private/MagicLeapHMD.cpp
@@ -2066,7 +2066,11 @@ bool FMagicLeapHMD::FMagicLeapVulkanExtensions::GetVulkanInstanceExtensionsRequi
 
 bool FMagicLeapHMD::FMagicLeapVulkanExtensions::GetVulkanDeviceExtensionsRequired(struct VkPhysicalDevice_T *pPhysicalDevice, TArray<const ANSICHAR*>& Out)
 {
+#if !PLATFORM_MAC
 	return FMagicLeapHelperVulkan::GetVulkanDeviceExtensionsRequired(pPhysicalDevice, Out);
+#else
+	return false;
+#endif
 }
 
 #undef LOCTEXT_NAMESPACE
