HandTracking LeftGestureButton fix
-The reality is Nihav made the fix, and I reviewed it.

diff --git a/Engine/Plugins/Lumin/MagicLeap/Source/MagicLeapHandTracking/Private/MagicLeapHandTracking.cpp b/Engine/Plugins/Lumin/MagicLeap/Source/MagicLeapHandTracking/Private/MagicLeapHandTracking.cpp
--- a/Engine/Plugins/Lumin/MagicLeap/Source/MagicLeapHandTracking/Private/MagicLeapHandTracking.cpp
+++ b/Engine/Plugins/Lumin/MagicLeap/Source/MagicLeapHandTracking/Private/MagicLeapHandTracking.cpp
@@ -669,7 +669,7 @@ void FMagicLeapHandTracking::SendControllerEventsForHand(const MLHandTrackingHan
 		FMagicLeapHMD::EnableInput EnableInputFromHMD;
 		// fixes unreferenced parameter error for Windows package builds.
 		(void)EnableInputFromHMD;
-		MessageHandler->OnControllerButtonReleased(RightStaticGestureMap[OldHandState.keypose], DeviceIndex, false);
+		MessageHandler->OnControllerButtonReleased(GestureMap[OldHandState.keypose], DeviceIndex, false);
 		if (GestureConfidenceThresholds[GestureIndex] <= NewConfidence)
 		{
 			MessageHandler->OnControllerButtonPressed(GestureMap[GestureIndex], DeviceIndex, false);
