HandTracking blueprint access to all keypoints
-Duplicating for Dev-VR from Dev-Partner-MagicLeap-4.20
-Created new GetGestureKeypointTransform blueprint function to get a keypoint's transform by hand and keypoint enum.
-Deprecated the old GetGestureKeypoint methods
-Fixed Special_1 and Special_2 to use hand center instead of wrist center.
-Exposed all the keypoints to blueprint, so they can work right away when underlyign support appears in the OS.
#MLQDR

diff --git a/Engine/Plugins/Lumin/MagicLeap/Source/MagicLeapHandTracking/Private/MagicLeapHandTracking.cpp b/Engine/Plugins/Lumin/MagicLeap/Source/MagicLeapHandTracking/Private/MagicLeapHandTracking.cpp
--- a/Engine/Plugins/Lumin/MagicLeap/Source/MagicLeapHandTracking/Private/MagicLeapHandTracking.cpp
+++ b/Engine/Plugins/Lumin/MagicLeap/Source/MagicLeapHandTracking/Private/MagicLeapHandTracking.cpp
@@ -76,6 +76,7 @@ IMPLEMENT_MODULE(FMagicLeapHandTrackingPlugin, MagicLeapHandTracking);
 
 // Partial sanity check of UE4 enums vs ML enums.  If these fail there has probably been an MLSDK update that requires an UE4 update.
 #if WITH_MLSDK
+static_assert(EHandTrackingKeypointCount == static_cast<int32>(MLHandTrackingStaticData_MaxKeyPoints), "EHandTrackingGesture does not match api enum.");
 static_assert(static_cast<int32>(EHandTrackingGesture::NoHand) == static_cast<int32>(MLHandTrackingKeyPose::MLHandTrackingKeyPose_NoHand), "EHandTrackingGesture does not match api enum.");
 static_assert(static_cast<int32>(EHandTrackingKeypointFilterLevel::PredictiveSmoothing) == static_cast<int32>(MLKeypointFilterLevel::MLKeypointFilterLevel_2), "EHandTrackingKeypointFilterLevel does not match api enum.");
 static_assert(static_cast<int32>(EHandTrackingGestureFilterLevel::MoreRobustnessToFlicker) == static_cast<int32>(MLPoseFilterLevel::MLPoseFilterLevel_2), "EHandTrackingGestureFilterLevel does not match api enum.");
@@ -327,7 +328,7 @@ FMagicLeapHandTracking::~FMagicLeapHandTracking()
 
 void FMagicLeapHandTracking::BuildKeypointMaps()
 {
-	SourceToTransformMap.Reserve(48);
+	SourceToTransformMap.Reserve(EHandTrackingKeypointCount * 2);
 
 	SourceToTransformMap.Add(LeftHandCenter_Name, &LeftHand.HandCenter);
 
@@ -394,6 +395,8 @@ void FMagicLeapHandTracking::BuildKeypointMaps()
 
 
 
+	SourceToHandMap.Reserve(EHandTrackingKeypointCount * 2);
+
 	SourceToHandMap.Add(LeftHandCenter_Name, &LeftHand);
 
 	SourceToHandMap.Add(LeftWristCenter_Name, &LeftHand);
@@ -456,7 +459,47 @@ void FMagicLeapHandTracking::BuildKeypointMaps()
 	SourceToHandMap.Add(RightPinkyFingerDIP_Name, &RightHand);
 	SourceToHandMap.Add(RightPinkyFingerPIP_Name, &RightHand);
 	SourceToHandMap.Add(RightPinkyFingerMCP_Name, &RightHand);
+}
 
+FMagicLeapHandTracking::FHandState::FHandState()
+{
+	EnumToTransformMap[static_cast<int32>(EHandTrackingKeypoint::Thumb_Tip)] = &Thumb.Tip;
+	EnumToTransformMap[static_cast<int32>(EHandTrackingKeypoint::Thumb_IP)] = &Thumb.DIP;
+	EnumToTransformMap[static_cast<int32>(EHandTrackingKeypoint::Thumb_MCP)] = &Thumb.PIP;
+	EnumToTransformMap[static_cast<int32>(EHandTrackingKeypoint::Thumb_CMC)] = &Thumb.MCP;
+
+	EnumToTransformMap[static_cast<int32>(EHandTrackingKeypoint::Index_Tip)] = &IndexFinger.Tip;
+	EnumToTransformMap[static_cast<int32>(EHandTrackingKeypoint::Index_DIP)] = &IndexFinger.DIP;
+	EnumToTransformMap[static_cast<int32>(EHandTrackingKeypoint::Index_PIP)] = &IndexFinger.PIP;
+	EnumToTransformMap[static_cast<int32>(EHandTrackingKeypoint::Index_MCP)] = &IndexFinger.MCP;
+
+	EnumToTransformMap[static_cast<int32>(EHandTrackingKeypoint::Middle_Tip)] = &MiddleFinger.Tip;
+	EnumToTransformMap[static_cast<int32>(EHandTrackingKeypoint::Middle_DIP)] = &MiddleFinger.DIP;
+	EnumToTransformMap[static_cast<int32>(EHandTrackingKeypoint::Middle_PIP)] = &MiddleFinger.PIP;
+	EnumToTransformMap[static_cast<int32>(EHandTrackingKeypoint::Middle_MCP)] = &MiddleFinger.MCP;
+
+	EnumToTransformMap[static_cast<int32>(EHandTrackingKeypoint::Ring_Tip)] = &RingFinger.Tip;
+	EnumToTransformMap[static_cast<int32>(EHandTrackingKeypoint::Ring_DIP)] = &RingFinger.DIP;
+	EnumToTransformMap[static_cast<int32>(EHandTrackingKeypoint::Ring_PIP)] = &RingFinger.PIP;
+	EnumToTransformMap[static_cast<int32>(EHandTrackingKeypoint::Ring_MCP)] = &RingFinger.MCP;
+
+	EnumToTransformMap[static_cast<int32>(EHandTrackingKeypoint::Pinky_Tip)] = &PinkyFinger.Tip;
+	EnumToTransformMap[static_cast<int32>(EHandTrackingKeypoint::Pinky_DIP)] = &PinkyFinger.DIP;
+	EnumToTransformMap[static_cast<int32>(EHandTrackingKeypoint::Pinky_PIP)] = &PinkyFinger.PIP;
+	EnumToTransformMap[static_cast<int32>(EHandTrackingKeypoint::Pinky_MCP)] = &PinkyFinger.MCP;
+
+	EnumToTransformMap[static_cast<int32>(EHandTrackingKeypoint::Wrist_Center)] = &Wrist.Center;
+	EnumToTransformMap[static_cast<int32>(EHandTrackingKeypoint::Wrist_Ulnar)] = &Wrist.Ulnar;
+	EnumToTransformMap[static_cast<int32>(EHandTrackingKeypoint::Wrist_Radial)] = &Wrist.Radial;
+	
+	EnumToTransformMap[static_cast<int32>(EHandTrackingKeypoint::Hand_Center)] = &HandCenter;
+}
+
+bool FMagicLeapHandTracking::FHandState::GetTransform(EHandTrackingKeypoint KeyPoint, FTransform& OutTransform) const
+{
+	check(static_cast<int32>(KeyPoint) >= 0 && static_cast<int32>(KeyPoint) <= EHandTrackingKeypointCount);
+	OutTransform = *EnumToTransformMap[static_cast<int32>(KeyPoint)];
+	return IsValid();
 }
 
 const FTransform* FMagicLeapHandTracking::FindTransformBySource(FName SourceName) const
@@ -509,7 +552,7 @@ bool FMagicLeapHandTracking::GetControllerOrientationAndPosition(const int32 Con
 			switch (DeviceHand)
 			{
 			case EControllerHand::Special_1:
-				ControllerTransform = &LeftHand.Wrist.Center;
+				ControllerTransform = &LeftHand.HandCenter;
 				break;
 			case EControllerHand::Special_3:
 				ControllerTransform = &LeftHand.IndexFinger.Tip;
@@ -518,7 +561,7 @@ bool FMagicLeapHandTracking::GetControllerOrientationAndPosition(const int32 Con
 				ControllerTransform = &LeftHand.Thumb.Tip;
 				break;
 			case EControllerHand::Special_2:
-				ControllerTransform = &RightHand.Wrist.Center;
+				ControllerTransform = &RightHand.HandCenter;
 				break;
 			case EControllerHand::Special_4:
 				ControllerTransform = &RightHand.IndexFinger.Tip;
@@ -577,64 +620,64 @@ void FMagicLeapHandTracking::EnumerateSources(TArray<FMotionControllerSource>& S
 	// Exposing only the keypoints that are actually tracked in MLSDK 0.15.0 RC5
 
 	SourcesOut.Add(LeftThumbTip_Name);
-	//SourcesOut.Add(LeftThumbIP_Name);
+	SourcesOut.Add(LeftThumbIP_Name);
 	SourcesOut.Add(LeftThumbMCP_Name);
-	//SourcesOut.Add(LeftThumbCMC_Name);
+	SourcesOut.Add(LeftThumbCMC_Name);
 
 	SourcesOut.Add(LeftIndexFingerTip_Name);
-	//SourcesOut.Add(LeftIndexFingerDIP_Name);
+	SourcesOut.Add(LeftIndexFingerDIP_Name);
 	SourcesOut.Add(LeftIndexFingerPIP_Name);
 	SourcesOut.Add(LeftIndexFingerMCP_Name);
 
-	//SourcesOut.Add(LeftMiddleFingerTip_Name);
-	//SourcesOut.Add(LeftMiddleFingerDIP_Name);
-	//SourcesOut.Add(LeftMiddleFingerPIP_Name);
+	SourcesOut.Add(LeftMiddleFingerTip_Name);
+	SourcesOut.Add(LeftMiddleFingerDIP_Name);
+	SourcesOut.Add(LeftMiddleFingerPIP_Name);
 	SourcesOut.Add(LeftMiddleFingerMCP_Name);
 
-	//SourcesOut.Add(LeftRingFingerTip_Name);
-	//SourcesOut.Add(LeftRingFingerDIP_Name);
-	//SourcesOut.Add(LeftRingFingerPIP_Name);
-	//SourcesOut.Add(LeftRingFingerMCP_Name);
+	SourcesOut.Add(LeftRingFingerTip_Name);
+	SourcesOut.Add(LeftRingFingerDIP_Name);
+	SourcesOut.Add(LeftRingFingerPIP_Name);
+	SourcesOut.Add(LeftRingFingerMCP_Name);
 
-	//SourcesOut.Add(LeftPinkyFingerTip_Name);
-	//SourcesOut.Add(LeftPinkyFingerDIP_Name);
-	//SourcesOut.Add(LeftPinkyFingerPIP_Name);
-	//SourcesOut.Add(LeftPinkyFingerMCP_Name);
+	SourcesOut.Add(LeftPinkyFingerTip_Name);
+	SourcesOut.Add(LeftPinkyFingerDIP_Name);
+	SourcesOut.Add(LeftPinkyFingerPIP_Name);
+	SourcesOut.Add(LeftPinkyFingerMCP_Name);
 
 	SourcesOut.Add(LeftWristCenter_Name);
-	//SourcesOut.Add(LeftWristUlnar_Name);
-	//SourcesOut.Add(LeftWristRadial_Name);
+	SourcesOut.Add(LeftWristUlnar_Name);
+	SourcesOut.Add(LeftWristRadial_Name);
 
 	SourcesOut.Add(LeftHandCenter_Name);
 
 
 	SourcesOut.Add(RightThumbTip_Name);
-	//SourcesOut.Add(RightThumbIP_Name);
+	SourcesOut.Add(RightThumbIP_Name);
 	SourcesOut.Add(RightThumbMCP_Name);
-	//SourcesOut.Add(RightThumbCMC_Name);
+	SourcesOut.Add(RightThumbCMC_Name);
 
 	SourcesOut.Add(RightIndexFingerTip_Name);
-	//SourcesOut.Add(RightIndexFingerDIP_Name);
+	SourcesOut.Add(RightIndexFingerDIP_Name);
 	SourcesOut.Add(RightIndexFingerPIP_Name);
 	SourcesOut.Add(RightIndexFingerMCP_Name);
 
-	//SourcesOut.Add(RightMiddleFingerTip_Name);
-	//SourcesOut.Add(RightMiddleFingerDIP_Name);
-	//SourcesOut.Add(RightMiddleFingerPIP_Name);
+	SourcesOut.Add(RightMiddleFingerTip_Name);
+	SourcesOut.Add(RightMiddleFingerDIP_Name);
+	SourcesOut.Add(RightMiddleFingerPIP_Name);
 	SourcesOut.Add(RightMiddleFingerMCP_Name);
 
-	//SourcesOut.Add(RightRingFingerTip_Name);
-	//SourcesOut.Add(RightRingFingerDIP_Name);
-	//SourcesOut.Add(RightRingFingerPIP_Name);
-	//SourcesOut.Add(RightRingFingerMCP_Name);
+	SourcesOut.Add(RightRingFingerTip_Name);
+	SourcesOut.Add(RightRingFingerDIP_Name);
+	SourcesOut.Add(RightRingFingerPIP_Name);
+	SourcesOut.Add(RightRingFingerMCP_Name);
 
-	//SourcesOut.Add(RightPinkyFingerTip_Name);
-	//SourcesOut.Add(RightPinkyFingerDIP_Name);
-	//SourcesOut.Add(RightPinkyFingerPIP_Name);
-	//SourcesOut.Add(RightPinkyFingerMCP_Name);
+	SourcesOut.Add(RightPinkyFingerTip_Name);
+	SourcesOut.Add(RightPinkyFingerDIP_Name);
+	SourcesOut.Add(RightPinkyFingerPIP_Name);
+	SourcesOut.Add(RightPinkyFingerMCP_Name);
 	SourcesOut.Add(RightWristCenter_Name);
-	//SourcesOut.Add(RightWristUlnar_Name);
-	//SourcesOut.Add(RightWristRadial_Name);
+	SourcesOut.Add(RightWristUlnar_Name);
+	SourcesOut.Add(RightWristRadial_Name);
 
 	SourcesOut.Add(RightHandCenter_Name);
 }
@@ -746,6 +789,13 @@ bool FMagicLeapHandTracking::IsHandTrackingStateValid() const
 	return bIsHandTrackingStateValid;
 }
 
+bool FMagicLeapHandTracking::GetKeypointTransform(EControllerHand Hand, EHandTrackingKeypoint Keypoint, FTransform& OutTransform) const
+{
+	const FMagicLeapHandTracking::FHandState& HandState = (Hand == EControllerHand::Left) ? GetLeftHandState() : GetRightHandState();
+
+	return HandState.GetTransform(Keypoint, OutTransform);
+}
+
 bool FMagicLeapHandTracking::SetConfiguration(bool bTrackingEnabled, const TArray<EHandTrackingGesture>& ActiveKeyPoses, EHandTrackingKeypointFilterLevel KeypointsFilterLevel, EHandTrackingGestureFilterLevel GestureFilterLevel)
 {
 #if WITH_MLSDK
diff --git a/Engine/Plugins/Lumin/MagicLeap/Source/MagicLeapHandTracking/Private/MagicLeapHandTracking.h b/Engine/Plugins/Lumin/MagicLeap/Source/MagicLeapHandTracking/Private/MagicLeapHandTracking.h
--- a/Engine/Plugins/Lumin/MagicLeap/Source/MagicLeapHandTracking/Private/MagicLeapHandTracking.h
+++ b/Engine/Plugins/Lumin/MagicLeap/Source/MagicLeapHandTracking/Private/MagicLeapHandTracking.h
@@ -110,6 +110,8 @@ public:
 	};
 	struct FHandState : public FNoncopyable
 	{
+		FHandState();
+
 		EHandTrackingGesture Gesture = EHandTrackingGesture::NoHand;
 		float GestureConfidence = 0.0f;
 		FVector HandCenterNormalized = FVector::ZeroVector;
@@ -125,6 +127,10 @@ public:
 		FTransform HandCenter = FTransform::Identity;
 
 		bool IsValid() const { return Gesture != EHandTrackingGesture::NoHand; }
+		bool GetTransform(EHandTrackingKeypoint KeyPoint, FTransform& OutTransform) const;
+
+	private:
+		FTransform* EnumToTransformMap[EHandTrackingKeypointCount];
 	};
 
 public:
@@ -155,6 +161,8 @@ public:
 	const FHandState& GetRightHandState() const;
 	bool IsHandTrackingStateValid() const;
 
+	bool GetKeypointTransform(EControllerHand Hand, EHandTrackingKeypoint Keypoint, FTransform& OutTransform) const;
+
 	bool SetConfiguration(bool bTrackingEnabled, const TArray<EHandTrackingGesture>& ActiveKeyPoses, EHandTrackingKeypointFilterLevel KeypointsFilterLevel, EHandTrackingGestureFilterLevel GestureFilterLevel);
 	bool GetConfiguration(bool& bTrackingEnabled, TArray<EHandTrackingGesture>& ActiveKeyPoses, EHandTrackingKeypointFilterLevel& KeypointsFilterLevel, EHandTrackingGestureFilterLevel& GestureFilterLevel);
 
diff --git a/Engine/Plugins/Lumin/MagicLeap/Source/MagicLeapHandTracking/Private/MagicLeapHandTrackingFunctionLibrary.cpp b/Engine/Plugins/Lumin/MagicLeap/Source/MagicLeapHandTracking/Private/MagicLeapHandTrackingFunctionLibrary.cpp
--- a/Engine/Plugins/Lumin/MagicLeap/Source/MagicLeapHandTracking/Private/MagicLeapHandTrackingFunctionLibrary.cpp
+++ b/Engine/Plugins/Lumin/MagicLeap/Source/MagicLeapHandTracking/Private/MagicLeapHandTrackingFunctionLibrary.cpp
@@ -172,6 +172,51 @@ bool UMagicLeapHandTrackingFunctionLibrary::GetGestureKeypoints(EControllerHand
 	return false;
 }
 
+bool UMagicLeapHandTrackingFunctionLibrary::GetGestureKeypointTransform(EControllerHand Hand, EHandTrackingKeypoint Keypoint, EGestureTransformSpace TransformSpace, FTransform& OutTransform)
+{
+	TSharedPtr<FMagicLeapHandTracking> HandTracking = StaticCastSharedPtr<FMagicLeapHandTracking>(IMagicLeapHandTrackingPlugin::Get().GetInputDevice());
+
+	if (HandTracking.IsValid() && HandTracking->IsHandTrackingStateValid())
+	{
+		FTransform KeyPointTransform;
+		const bool bSuccess = HandTracking->GetKeypointTransform(Hand, Keypoint, KeyPointTransform);
+
+		if (bSuccess)
+		{
+			switch (TransformSpace)
+			{
+			case EGestureTransformSpace::Tracking:
+			{
+				OutTransform = KeyPointTransform;
+				return true;
+			}
+			case EGestureTransformSpace::World:
+			{
+				FTransform TrackingToWorldTransform = UHeadMountedDisplayFunctionLibrary::GetTrackingToWorldTransform(GWorld);
+				OutTransform = KeyPointTransform * TrackingToWorldTransform;
+				return true;
+			}
+			case EGestureTransformSpace::Hand:
+			{
+				FTransform HandTransform;
+				const bool bSuccess2 = HandTracking->GetKeypointTransform(Hand, EHandTrackingKeypoint::Hand_Center, HandTransform);
+				if (bSuccess2)
+				{
+					OutTransform = KeyPointTransform * HandTransform.Inverse();
+					return true;
+				}
+				break;
+			}
+			default:
+				check(false);
+				return false;
+			}
+		}
+	}
+
+	return false;
+}
+
 bool UMagicLeapHandTrackingFunctionLibrary::SetConfiguration(const TArray<EHandTrackingGesture>& StaticGesturesToActivate, EHandTrackingKeypointFilterLevel KeypointsFilterLevel, EHandTrackingGestureFilterLevel GestureFilterLevel, EHandTrackingGestureFilterLevel HandSwitchingFilterLevel, bool bEnabled)
 {
 	TSharedPtr<FMagicLeapHandTracking> HandTracking = StaticCastSharedPtr<FMagicLeapHandTracking>(IMagicLeapHandTrackingPlugin::Get().GetInputDevice());
diff --git a/Engine/Plugins/Lumin/MagicLeap/Source/MagicLeapHandTracking/Public/MagicLeapHandTrackingFunctionLibrary.h b/Engine/Plugins/Lumin/MagicLeap/Source/MagicLeapHandTracking/Public/MagicLeapHandTrackingFunctionLibrary.h
--- a/Engine/Plugins/Lumin/MagicLeap/Source/MagicLeapHandTracking/Public/MagicLeapHandTrackingFunctionLibrary.h
+++ b/Engine/Plugins/Lumin/MagicLeap/Source/MagicLeapHandTracking/Public/MagicLeapHandTrackingFunctionLibrary.h
@@ -19,26 +19,24 @@ public:
 	/**
 	  Transform of the center of the hand.  Approximately the center of the palm.
 
-	  Note that this returns a transform in the Tracking space. To get the transform in Unreal's
-	  world space, use the MotioController component as a child of the XRPawn with hand set to EControllerHand::Special_1
-	  for the left hand center and EControllerHand::Special_2 for the right hand center.
+	  Note that this returns a transform in world space.
 
 	  @param Hand Hand to query the hand center transform for. Only Left and Right hand are supported.
 	  @param HandCenter Output parameter containing the position and orientation of the given hand.
-	  @return true if the output param was populated with a valid value, false otherwise.
+	  @return true if the output param was populated with a valid value, false means that is is either unchanged or populated with a stale value.
 	*/
-	UFUNCTION(BlueprintCallable, Category = "HandTracking|MagicLeap")
+	UFUNCTION(BlueprintCallable, Category = "HandTracking|MagicLeap", meta = (DeprecatedFunction, DeprecationMessage = "GetHandCenter is deprecated. Please use GetGestureKeypointTransform instead."))
 	static bool GetHandCenter(EControllerHand Hand, FTransform& HandCenter);
 
 	/**
 	  Transform of the Index Finger Tip.
 
 	  @param Hand Hand to query the hand center transform for. Only Left and Right hand are supported.
-	  @param TransformSpace Get the transform relative to the world, tracking space, or the hand center.
+	  @param TransformSpace Get the transform relative to the world, hand center, or tracking space.
 	  @param Pointer Output parameter containing the position and orientation.
-	  @return true if the output param was populated with a valid value, false otherwise.
-	*/
-	UFUNCTION(BlueprintCallable, Category = "HandTracking|MagicLeap")
+	  @return true if the output param was populated with a valid value, false means that is is either unchanged or populated with a stale value.
+	  */
+	UFUNCTION(BlueprintCallable, Category = "HandTracking|MagicLeap", meta = (DeprecatedFunction, DeprecationMessage = "GetHandIndexFingerTip is deprecated. Please use GetGestureKeypointTransform instead."))
 	static bool GetHandIndexFingerTip(EControllerHand Hand, EGestureTransformSpace TransformSpace, FTransform& Pointer);
 
 	/**
@@ -49,11 +47,11 @@ public:
 	  for the left hand secondary and EControllerHand::Special_6 for the right hand secondary.
 
 	  @param Hand Hand to query the hand center transform for. Only Left and Right hand are supported.
-	  @param TransformSpace Get the transform relative to the world, tracking space, or the hand center.
+	  @param TransformSpace Get the transform relative to the world, hand center, or tracking space.
 	  @param Secondary Output parameter containing the position and orientation.
-	  @return true if the output param was populated with a valid value, false otherwise.
-	*/
-	UFUNCTION(BlueprintCallable, Category = "HandTracking|MagicLeap")
+	  @return true if the output param was populated with a valid value, false means that is is either unchanged or populated with a stale value.
+	  */
+	UFUNCTION(BlueprintCallable, Category = "HandTracking|MagicLeap", meta = (DeprecatedFunction, DeprecationMessage = "GetHandThumbTip is deprecated. Please use GetGestureKeypointTransform instead."))
 	static bool GetHandThumbTip(EControllerHand Hand, EGestureTransformSpace TransformSpace, FTransform& Secondary);
 
 	/**
@@ -71,18 +69,30 @@ public:
 
 	  Note that this returns a transform in the Tracking space. To get the transform in Unreal's
 	  world space, use the MotioController component as a child of the XRPawn with hand set to the following.
-	  Special_3 - Left keypoint 0
-	  Special_5 - Left keypoint 1
-	  Special_4 - Right keypoint 0
-	  Special_6 - Right keypoint 1
+	  Special_3 - Left Index Finger Tip
+	  Special_5 - Left Thumb Tip
+	  Special_4 - Right Index Finger Tip
+	  Special_6 - Right Thumb Tip
 
 	  @param Hand Hand to query the keypoints for. Only Left and Right hand are supported.
-	  @param Keypoints Output paramter containing transforms of the keypoints detected on the given hand.
+	  @param Keypoints Output parameter containing transforms of the keypoints detected on the given hand.
 	  @return true if the output param was populated with a valid value, false otherwise.
 	*/
-	UFUNCTION(BlueprintCallable, Category = "HandTracking|MagicLeap")
+	UFUNCTION(BlueprintCallable, Category = "HandTracking|MagicLeap", meta = (DeprecatedFunction, DeprecationMessage = "GetGestureKeypoints is deprecated. Please use GetGestureKeypointTransform instead."))
 	static bool GetGestureKeypoints(EControllerHand Hand, TArray<FTransform>& Keypoints);
 
+	/**
+	Get Transform for a point on the hand.
+
+	@param Hand
+	@param Keypoint the specific joint or wrist point to fetch.
+	@param Transform Output parameter to write the data to.
+	@param TransformSpace Get the transform relative to the world, hand center, or tracking space.
+	@return true if the output param was populated with a valid value, false means that is is either unchanged or populated with a stale value.
+	*/
+	UFUNCTION(BlueprintCallable, Category = "HandTracking|MagicLeap")
+	static bool GetGestureKeypointTransform(EControllerHand Hand, EHandTrackingKeypoint Keypoint, EGestureTransformSpace TransformSpace, FTransform& Transform);
+
 	/**
 	  Enables and disables the gestures to be detected by the gesture recognition system.
 
diff --git a/Engine/Plugins/Lumin/MagicLeap/Source/MagicLeapHandTracking/Public/MagicLeapHandTrackingTypes.h b/Engine/Plugins/Lumin/MagicLeap/Source/MagicLeapHandTracking/Public/MagicLeapHandTrackingTypes.h
--- a/Engine/Plugins/Lumin/MagicLeap/Source/MagicLeapHandTracking/Public/MagicLeapHandTrackingTypes.h
+++ b/Engine/Plugins/Lumin/MagicLeap/Source/MagicLeapHandTracking/Public/MagicLeapHandTrackingTypes.h
@@ -6,6 +6,44 @@
 #include "UObject/ObjectMacros.h"
 #include "MagicLeapHandTrackingTypes.generated.h"
 
+/*! Transforms that could be tracked on the hand. In 0.15.0 RC5 8 of them may be tracked.*/
+UENUM(BlueprintType)
+enum class EHandTrackingKeypoint : uint8
+{
+	Thumb_Tip, // SDK 0.15.0 tracked
+	Thumb_IP,
+	Thumb_MCP, // SDK 0.15.0 tracked
+	Thumb_CMC,
+
+	Index_Tip, // SDK 0.15.0 tracked
+	Index_DIP,
+	Index_PIP, // SDK 0.15.0 tracked
+	Index_MCP, // SDK 0.15.0 tracked
+
+	Middle_Tip,
+	Middle_DIP,
+	Middle_PIP,
+	Middle_MCP, // SDK 0.15.0 tracked
+
+	Ring_Tip,
+	Ring_DIP,
+	Ring_PIP,
+	Ring_MCP,
+
+	Pinky_Tip,
+	Pinky_DIP,
+	Pinky_PIP,
+	Pinky_MCP,
+
+	Wrist_Center, // SDK 0.15.0 tracked
+	Wrist_Ulnar,
+	Wrist_Radial,
+
+	Hand_Center // SDK 0.15.0 tracked
+};
+
+const int32 EHandTrackingKeypointCount = static_cast<int32>(EHandTrackingKeypoint::Hand_Center) + 1;
+
 /*! Static key pose types which are available when both hands are separated. */
 UENUM(BlueprintType)
 enum class EHandTrackingGesture : uint8
@@ -60,12 +98,12 @@ enum class EHandTrackingGestureFilterLevel : uint8
 UENUM(BlueprintType)
 enum class EGestureTransformSpace : uint8
 {
-	/** Gesture key point transforms are reported in device Tracking space. */
-	Tracking,
 	/** Gesture key point transforms are reported in Unreal world space. This is more costly on CPU.*/
 	World,
 	/** Gesture key point transforms are reported in gesture hand center space.*/
-	Hand
+	Hand,
+	/** Gesture key point transforms are reported in device Tracking space. */
+	Tracking
 };
 
 /** List of input key names for all left and right hand gestures. */
