UEVR-1261 MagicLeap HandTracking should be exposed to LiveLink
-Exposed hand tracking transforms through live link.
-Added blueprint function to get the livelink source.
-Exposed LiveLikeSourceHandle through ILiveLinkSource.h so that it can be used by the MagicLeapPlugin.
-The transforms are in tracking space.  They form a hierarchical skeleton with HandCenter as the root.  See FMagicLeapHandTracking::SetupLiveLinkData() for details on the hierarchy.  With the current magicleap runtime functionality many transforms in the skeleton are always identity.

diff --git a/Engine/Plugins/Animation/LiveLink/Source/LiveLink/Private/LiveLinkBlueprintStructs.cpp b/Engine/Plugins/Animation/LiveLink/Source/LiveLink/Private/LiveLinkBlueprintStructs.cpp
--- a/Engine/Plugins/Animation/LiveLink/Source/LiveLink/Private/LiveLinkBlueprintStructs.cpp
+++ b/Engine/Plugins/Animation/LiveLink/Source/LiveLink/Private/LiveLinkBlueprintStructs.cpp
@@ -343,8 +343,3 @@ void FSubjectFrameHandle::SetCachedFrame(TSharedPtr<FCachedSubjectFrame> InCache
 {
 	CachedFrame = InCachedFrame;
 };
-
-void FLiveLinkSourceHandle::SetSourcePointer(TSharedPtr<ILiveLinkSource> InSourcePointer)
-{
-	SourcePointer = InSourcePointer;
-};
\ No newline at end of file
diff --git a/Engine/Plugins/Animation/LiveLink/Source/LiveLink/Public/LiveLinkBlueprintStructs.h b/Engine/Plugins/Animation/LiveLink/Source/LiveLink/Public/LiveLinkBlueprintStructs.h
--- a/Engine/Plugins/Animation/LiveLink/Source/LiveLink/Public/LiveLinkBlueprintStructs.h
+++ b/Engine/Plugins/Animation/LiveLink/Source/LiveLink/Public/LiveLinkBlueprintStructs.h
@@ -132,17 +132,3 @@ private:
 	TSharedPtr<FCachedSubjectFrame> CachedFrame;
 };
 
-// A Blueprint handle to a specific LiveLink Source
-USTRUCT(BlueprintType)
-struct FLiveLinkSourceHandle
-{
-	GENERATED_USTRUCT_BODY()
-	
-	FLiveLinkSourceHandle() = default;
-
-	virtual ~FLiveLinkSourceHandle() = default;
-
-	void SetSourcePointer(TSharedPtr<ILiveLinkSource> InSourcePointer);
-
-	TSharedPtr<ILiveLinkSource> SourcePointer;
-};
diff --git a/Engine/Plugins/Lumin/MagicLeap/Source/MagicLeapHandTracking/MagicLeapHandTracking.Build.cs b/Engine/Plugins/Lumin/MagicLeap/Source/MagicLeapHandTracking/MagicLeapHandTracking.Build.cs
--- a/Engine/Plugins/Lumin/MagicLeap/Source/MagicLeapHandTracking/MagicLeapHandTracking.Build.cs
+++ b/Engine/Plugins/Lumin/MagicLeap/Source/MagicLeapHandTracking/MagicLeapHandTracking.Build.cs
@@ -28,7 +28,8 @@ namespace UnrealBuildTool.Rules
 					"LuminRuntimeSettings",
 					"MagicLeap",
 					"MLSDK",
-					"SlateCore"
+					"SlateCore",
+					"LiveLinkInterface"
 				}
 			);
 
diff --git a/Engine/Plugins/Lumin/MagicLeap/Source/MagicLeapHandTracking/Private/LiveLinkMagicLeapHandTrackingSourceEditor.cpp b/Engine/Plugins/Lumin/MagicLeap/Source/MagicLeapHandTracking/Private/LiveLinkMagicLeapHandTrackingSourceEditor.cpp
new file mode 100644
--- /dev/null
+++ b/Engine/Plugins/Lumin/MagicLeap/Source/MagicLeapHandTracking/Private/LiveLinkMagicLeapHandTrackingSourceEditor.cpp
@@ -0,0 +1,15 @@
+// Copyright 1998-2018 Epic Games, Inc. All Rights Reserved.
+
+#include "LiveLinkMagicLeapHandTrackingSourceEditor.h"
+
+#define LOCTEXT_NAMESPACE "LiveLinkMagicLeapHandTrackingSourceEditor"
+
+SLiveLinkMagicLeapHandTrackingSourceEditor::~SLiveLinkMagicLeapHandTrackingSourceEditor()
+{
+}
+
+void SLiveLinkMagicLeapHandTrackingSourceEditor::Construct(const FArguments& Args)
+{
+}
+
+#undef LOCTEXT_NAMESPACE
diff --git a/Engine/Plugins/Lumin/MagicLeap/Source/MagicLeapHandTracking/Private/LiveLinkMagicLeapHandTrackingSourceEditor.h b/Engine/Plugins/Lumin/MagicLeap/Source/MagicLeapHandTracking/Private/LiveLinkMagicLeapHandTrackingSourceEditor.h
new file mode 100644
--- /dev/null
+++ b/Engine/Plugins/Lumin/MagicLeap/Source/MagicLeapHandTracking/Private/LiveLinkMagicLeapHandTrackingSourceEditor.h
@@ -0,0 +1,31 @@
+// Copyright 1998-2018 Epic Games, Inc. All Rights Reserved.
+
+#pragma once
+
+#include "Widgets/SCompoundWidget.h"
+#include "Widgets/Views/SListView.h"
+#include "IMessageContext.h"
+#include "MessageEndpoint.h"
+#include "Widgets/DeclarativeSyntaxSupport.h"
+#include "Misc/Guid.h"
+
+struct FLiveLinkPongMessage;
+struct FMessageAddress;
+struct FProviderPollResult;
+class ITableRow;
+class STableViewBase;
+
+typedef TSharedPtr<FProviderPollResult> FProviderPollResultPtr;
+
+class SLiveLinkMagicLeapHandTrackingSourceEditor : public SCompoundWidget
+{
+	SLATE_BEGIN_ARGS(SLiveLinkMagicLeapHandTrackingSourceEditor)
+	{
+	}
+
+	SLATE_END_ARGS()
+
+	~SLiveLinkMagicLeapHandTrackingSourceEditor();
+
+	void Construct(const FArguments& Args);
+};
\ No newline at end of file
diff --git a/Engine/Plugins/Lumin/MagicLeap/Source/MagicLeapHandTracking/Private/LiveLinkMagicLeapHandTrackingSourceFactory.cpp b/Engine/Plugins/Lumin/MagicLeap/Source/MagicLeapHandTracking/Private/LiveLinkMagicLeapHandTrackingSourceFactory.cpp
new file mode 100644
--- /dev/null
+++ b/Engine/Plugins/Lumin/MagicLeap/Source/MagicLeapHandTracking/Private/LiveLinkMagicLeapHandTrackingSourceFactory.cpp
@@ -0,0 +1,44 @@
+﻿// Copyright 1998-2018 Epic Games, Inc. All Rights Reserved.
+#include "LiveLinkMagicLeapHandTrackingSourceFactory.h"
+#include "LiveLinkMagicLeapHandTrackingSourceEditor.h"
+#include "IMagicLeapHandTrackingPlugin.h"
+#include "MagicLeapHandTracking.h"
+
+#define LOCTEXT_NAMESPACE "MagicLeapHandTracking"
+
+FText ULiveLinkMagicLeapHandTrackingSourceFactory::GetSourceDisplayName() const
+{
+	return LOCTEXT("HandTrackingLiveLinkSourceName", "Hand Tracking Source");
+}
+
+FText ULiveLinkMagicLeapHandTrackingSourceFactory::GetSourceTooltip() const
+{
+	return LOCTEXT("HandTrackingLiveLinkSourceTooltip", "Hand Tracking Key Points Source");
+}
+
+TSharedPtr<SWidget> ULiveLinkMagicLeapHandTrackingSourceFactory::CreateSourceCreationPanel()
+{
+	if (!ActiveSourceEditor.IsValid())
+	{
+		SAssignNew(ActiveSourceEditor, SLiveLinkMagicLeapHandTrackingSourceEditor);
+	}
+	return ActiveSourceEditor;
+}
+
+TSharedPtr<ILiveLinkSource> ULiveLinkMagicLeapHandTrackingSourceFactory::OnSourceCreationPanelClosed(bool bCreateSource)
+{
+	TSharedPtr<ILiveLinkSource> NewSource = nullptr;
+
+	if (bCreateSource && ActiveSourceEditor.IsValid())
+	{
+		TSharedPtr<FMagicLeapHandTracking> HandTracking = StaticCastSharedPtr<FMagicLeapHandTracking>(IMagicLeapHandTrackingPlugin::Get().GetLiveLinkSource());
+
+		// Here we could apply settings from SLiveLinkMagicLeapHandTrackingSourceEditor
+
+		NewSource = HandTracking;
+	}
+	ActiveSourceEditor = nullptr;
+	return NewSource;
+}
+
+#undef LOCTEXT_NAMESPACE
\ No newline at end of file
diff --git a/Engine/Plugins/Lumin/MagicLeap/Source/MagicLeapHandTracking/Private/LiveLinkMagicLeapHandTrackingSourceFactory.h b/Engine/Plugins/Lumin/MagicLeap/Source/MagicLeapHandTracking/Private/LiveLinkMagicLeapHandTrackingSourceFactory.h
new file mode 100644
--- /dev/null
+++ b/Engine/Plugins/Lumin/MagicLeap/Source/MagicLeapHandTracking/Private/LiveLinkMagicLeapHandTrackingSourceFactory.h
@@ -0,0 +1,22 @@
+// Copyright 1998-2018 Epic Games, Inc. All Rights Reserved.
+
+#pragma once
+
+#include "LiveLinkSourceFactory.h"
+#include "LiveLinkMagicLeapHandTrackingSourceFactory.generated.h"
+
+UCLASS()
+class ULiveLinkMagicLeapHandTrackingSourceFactory : public ULiveLinkSourceFactory
+{
+public:
+
+	GENERATED_BODY()
+
+	virtual FText GetSourceDisplayName() const;
+	virtual FText GetSourceTooltip() const;
+
+	virtual TSharedPtr<SWidget> CreateSourceCreationPanel();
+	virtual TSharedPtr<ILiveLinkSource> OnSourceCreationPanelClosed(bool bMakeSource);
+
+	TSharedPtr<class SLiveLinkMagicLeapHandTrackingSourceEditor> ActiveSourceEditor;
+};
\ No newline at end of file
diff --git a/Engine/Plugins/Lumin/MagicLeap/Source/MagicLeapHandTracking/Private/MagicLeapHandTracking.cpp b/Engine/Plugins/Lumin/MagicLeap/Source/MagicLeapHandTracking/Private/MagicLeapHandTracking.cpp
--- a/Engine/Plugins/Lumin/MagicLeap/Source/MagicLeapHandTracking/Private/MagicLeapHandTracking.cpp
+++ b/Engine/Plugins/Lumin/MagicLeap/Source/MagicLeapHandTracking/Private/MagicLeapHandTracking.cpp
@@ -45,7 +45,7 @@ public:
 	{
 		if (!InputDevice.IsValid())
 		{
-			TSharedPtr<IInputDevice> HandTrackingInputDevice(new FMagicLeapHandTracking(InMessageHandler));
+			TSharedPtr<FMagicLeapHandTracking> HandTrackingInputDevice(new FMagicLeapHandTracking(InMessageHandler));
 			InputDevice = HandTrackingInputDevice;
 			return InputDevice;
 		}
@@ -61,13 +61,22 @@ public:
 	{
 		if (!InputDevice.IsValid())
 		{
-			InputDevice = CreateInputDevice(FSlateApplication::Get().GetPlatformApplication()->GetMessageHandler());
+			CreateInputDevice(FSlateApplication::Get().GetPlatformApplication()->GetMessageHandler());
+		}
+		return InputDevice;
+	}
+
+	virtual TSharedPtr<ILiveLinkSource> GetLiveLinkSource() override
+	{
+		if (!InputDevice.IsValid())
+		{
+			CreateInputDevice(FSlateApplication::Get().GetPlatformApplication()->GetMessageHandler());
 		}
 		return InputDevice;
 	}
 
 private:
-	TSharedPtr<IInputDevice> InputDevice;
+	TSharedPtr<FMagicLeapHandTracking> InputDevice;
 };
 
 IMPLEMENT_MODULE(FMagicLeapHandTrackingPlugin, MagicLeapHandTracking);
@@ -144,6 +153,41 @@ const FName FMagicLeapGestureKeyNames::Right_NoPose_Name("MagicLeap_Right_NoPose
 const FKey FMagicLeapHandTracking::FStaticHandTracking::Right_NoHand("MagicLeap_Right_NoHand");
 const FName FMagicLeapGestureKeyNames::Right_NoHand_Name("MagicLeap_Right_NoHand");
 
+// KeyPoint names
+#define DEFINE_HAND_KEY_POINT_NAME(name) FName FMagicLeapHandTracking::name##_Name(TEXT(#name));
+
+DEFINE_HAND_KEY_POINT_NAME(HandCenter);
+
+DEFINE_HAND_KEY_POINT_NAME(WristCenter);
+DEFINE_HAND_KEY_POINT_NAME(WristUlnar);
+DEFINE_HAND_KEY_POINT_NAME(WristRadial);
+
+DEFINE_HAND_KEY_POINT_NAME(ThumbTip);
+DEFINE_HAND_KEY_POINT_NAME(ThumbIP);
+DEFINE_HAND_KEY_POINT_NAME(ThumbMCP);
+DEFINE_HAND_KEY_POINT_NAME(ThumbCMC);
+
+DEFINE_HAND_KEY_POINT_NAME(IndexFingerTip);
+DEFINE_HAND_KEY_POINT_NAME(IndexFingerDIP);
+DEFINE_HAND_KEY_POINT_NAME(IndexFingerPIP);
+DEFINE_HAND_KEY_POINT_NAME(IndexFingerMCP);
+
+DEFINE_HAND_KEY_POINT_NAME(MiddleFingerTip);
+DEFINE_HAND_KEY_POINT_NAME(MiddleFingerDIP);
+DEFINE_HAND_KEY_POINT_NAME(MiddleFingerPIP);
+DEFINE_HAND_KEY_POINT_NAME(MiddleFingerMCP);
+
+DEFINE_HAND_KEY_POINT_NAME(RingFingerTip);
+DEFINE_HAND_KEY_POINT_NAME(RingFingerDIP);
+DEFINE_HAND_KEY_POINT_NAME(RingFingerPIP);
+DEFINE_HAND_KEY_POINT_NAME(RingFingerMCP);
+
+DEFINE_HAND_KEY_POINT_NAME(PinkyFingerTip);
+DEFINE_HAND_KEY_POINT_NAME(PinkyFingerDIP);
+DEFINE_HAND_KEY_POINT_NAME(PinkyFingerPIP);
+DEFINE_HAND_KEY_POINT_NAME(PinkyFingerMCP);
+
+#undef DEFINE_HAND_KEY_POINT_NAME
 
 // KeyPoint names
 #define DEFINE_KEY_POINT_NAME(name) FName FMagicLeapHandTracking::name##_Name(TEXT(#name));
@@ -211,6 +255,12 @@ DEFINE_KEY_POINT_NAME(RightPinkyFingerDIP);
 DEFINE_KEY_POINT_NAME(RightPinkyFingerPIP);
 DEFINE_KEY_POINT_NAME(RightPinkyFingerMCP);
 
+#undef DEFINE_KEY_POINT_NAME
+
+FName FMagicLeapHandTracking::LiveLinkLeftHandTrackingSubjectName(TEXT("MagicLeapLeftHandTracking"));
+FName FMagicLeapHandTracking::LiveLinkRightHandTrackingSubjectName(TEXT("MagicLeapRightHandTracking"));
+
+
 #if WITH_MLSDK
 EHandTrackingKeypointFilterLevel MLToUnrealKeypointsFilterLevel(MLKeypointFilterLevel FilterLevel)
 {
@@ -293,6 +343,7 @@ FMagicLeapHandTracking::FMagicLeapHandTracking(const TSharedRef<FGenericApplicat
 #endif //WITH_MLSDK
 
 	BuildKeypointMaps();
+	SetupLiveLinkData();
 
 	// Register "MotionController" modular feature manually
 	IModularFeatures::Get().RegisterModularFeature(GetModularFeatureName(), this);
@@ -330,68 +381,68 @@ void FMagicLeapHandTracking::BuildKeypointMaps()
 {
 	SourceToTransformMap.Reserve(EHandTrackingKeypointCount * 2);
 
-	SourceToTransformMap.Add(LeftHandCenter_Name, &LeftHand.HandCenter);
+	SourceToTransformMap.Add(LeftHandCenter_Name, &LeftHand.HandCenter.Transform);
 
-	SourceToTransformMap.Add(LeftWristCenter_Name, &LeftHand.Wrist.Center);
-	SourceToTransformMap.Add(LeftWristUlnar_Name, &LeftHand.Wrist.Ulnar);
-	SourceToTransformMap.Add(LeftWristRadial_Name, &LeftHand.Wrist.Radial);
+	SourceToTransformMap.Add(LeftWristCenter_Name, &LeftHand.Wrist.Center.Transform);
+	SourceToTransformMap.Add(LeftWristUlnar_Name, &LeftHand.Wrist.Ulnar.Transform);
+	SourceToTransformMap.Add(LeftWristRadial_Name, &LeftHand.Wrist.Radial.Transform);
 
-	SourceToTransformMap.Add(LeftThumbTip_Name, &LeftHand.Thumb.Tip);
-	SourceToTransformMap.Add(LeftThumbIP_Name, &LeftHand.Thumb.DIP);
-	SourceToTransformMap.Add(LeftThumbMCP_Name, &LeftHand.Thumb.PIP);
-	SourceToTransformMap.Add(LeftThumbCMC_Name, &LeftHand.Thumb.MCP);
+	SourceToTransformMap.Add(LeftThumbTip_Name, &LeftHand.Thumb.Tip.Transform);
+	SourceToTransformMap.Add(LeftThumbIP_Name, &LeftHand.Thumb.DIP.Transform);
+	SourceToTransformMap.Add(LeftThumbMCP_Name, &LeftHand.Thumb.PIP.Transform);
+	SourceToTransformMap.Add(LeftThumbCMC_Name, &LeftHand.Thumb.MCP.Transform);
 
-	SourceToTransformMap.Add(LeftIndexFingerTip_Name, &LeftHand.IndexFinger.Tip);
-	SourceToTransformMap.Add(LeftIndexFingerDIP_Name, &LeftHand.IndexFinger.DIP);
-	SourceToTransformMap.Add(LeftIndexFingerPIP_Name, &LeftHand.IndexFinger.PIP);
-	SourceToTransformMap.Add(LeftIndexFingerMCP_Name, &LeftHand.IndexFinger.MCP);
+	SourceToTransformMap.Add(LeftIndexFingerTip_Name, &LeftHand.IndexFinger.Tip.Transform);
+	SourceToTransformMap.Add(LeftIndexFingerDIP_Name, &LeftHand.IndexFinger.DIP.Transform);
+	SourceToTransformMap.Add(LeftIndexFingerPIP_Name, &LeftHand.IndexFinger.PIP.Transform);
+	SourceToTransformMap.Add(LeftIndexFingerMCP_Name, &LeftHand.IndexFinger.MCP.Transform);
 
-	SourceToTransformMap.Add(LeftMiddleFingerTip_Name, &LeftHand.MiddleFinger.Tip);
-	SourceToTransformMap.Add(LeftMiddleFingerDIP_Name, &LeftHand.MiddleFinger.DIP);
-	SourceToTransformMap.Add(LeftMiddleFingerPIP_Name, &LeftHand.MiddleFinger.PIP);
-	SourceToTransformMap.Add(LeftMiddleFingerMCP_Name, &LeftHand.MiddleFinger.MCP);
+	SourceToTransformMap.Add(LeftMiddleFingerTip_Name, &LeftHand.MiddleFinger.Tip.Transform);
+	SourceToTransformMap.Add(LeftMiddleFingerDIP_Name, &LeftHand.MiddleFinger.DIP.Transform);
+	SourceToTransformMap.Add(LeftMiddleFingerPIP_Name, &LeftHand.MiddleFinger.PIP.Transform);
+	SourceToTransformMap.Add(LeftMiddleFingerMCP_Name, &LeftHand.MiddleFinger.MCP.Transform);
 
-	SourceToTransformMap.Add(LeftRingFingerTip_Name, &LeftHand.RingFinger.Tip);
-	SourceToTransformMap.Add(LeftRingFingerDIP_Name, &LeftHand.RingFinger.DIP);
-	SourceToTransformMap.Add(LeftRingFingerPIP_Name, &LeftHand.RingFinger.PIP);
-	SourceToTransformMap.Add(LeftRingFingerMCP_Name, &LeftHand.RingFinger.MCP);
+	SourceToTransformMap.Add(LeftRingFingerTip_Name, &LeftHand.RingFinger.Tip.Transform);
+	SourceToTransformMap.Add(LeftRingFingerDIP_Name, &LeftHand.RingFinger.DIP.Transform);
+	SourceToTransformMap.Add(LeftRingFingerPIP_Name, &LeftHand.RingFinger.PIP.Transform);
+	SourceToTransformMap.Add(LeftRingFingerMCP_Name, &LeftHand.RingFinger.MCP.Transform);
 
-	SourceToTransformMap.Add(LeftPinkyFingerTip_Name, &LeftHand.PinkyFinger.Tip);
-	SourceToTransformMap.Add(LeftPinkyFingerDIP_Name, &LeftHand.PinkyFinger.DIP);
-	SourceToTransformMap.Add(LeftPinkyFingerPIP_Name, &LeftHand.PinkyFinger.PIP);
-	SourceToTransformMap.Add(LeftPinkyFingerMCP_Name, &LeftHand.PinkyFinger.MCP);
+	SourceToTransformMap.Add(LeftPinkyFingerTip_Name, &LeftHand.PinkyFinger.Tip.Transform);
+	SourceToTransformMap.Add(LeftPinkyFingerDIP_Name, &LeftHand.PinkyFinger.DIP.Transform);
+	SourceToTransformMap.Add(LeftPinkyFingerPIP_Name, &LeftHand.PinkyFinger.PIP.Transform);
+	SourceToTransformMap.Add(LeftPinkyFingerMCP_Name, &LeftHand.PinkyFinger.MCP.Transform);
 
 
-	SourceToTransformMap.Add(RightHandCenter_Name, &RightHand.HandCenter);
+	SourceToTransformMap.Add(RightHandCenter_Name, &RightHand.HandCenter.Transform);
 
-	SourceToTransformMap.Add(RightWristCenter_Name, &RightHand.Wrist.Center);
-	SourceToTransformMap.Add(RightWristUlnar_Name, &RightHand.Wrist.Ulnar);
-	SourceToTransformMap.Add(RightWristRadial_Name, &RightHand.Wrist.Radial);
+	SourceToTransformMap.Add(RightWristCenter_Name, &RightHand.Wrist.Center.Transform);
+	SourceToTransformMap.Add(RightWristUlnar_Name, &RightHand.Wrist.Ulnar.Transform);
+	SourceToTransformMap.Add(RightWristRadial_Name, &RightHand.Wrist.Radial.Transform);
 
-	SourceToTransformMap.Add(RightThumbTip_Name, &RightHand.Thumb.Tip);
-	SourceToTransformMap.Add(RightThumbIP_Name, &RightHand.Thumb.DIP);
-	SourceToTransformMap.Add(RightThumbMCP_Name, &RightHand.Thumb.PIP);
-	SourceToTransformMap.Add(RightThumbCMC_Name, &RightHand.Thumb.MCP);
+	SourceToTransformMap.Add(RightThumbTip_Name, &RightHand.Thumb.Tip.Transform);
+	SourceToTransformMap.Add(RightThumbIP_Name, &RightHand.Thumb.DIP.Transform);
+	SourceToTransformMap.Add(RightThumbMCP_Name, &RightHand.Thumb.PIP.Transform);
+	SourceToTransformMap.Add(RightThumbCMC_Name, &RightHand.Thumb.MCP.Transform);
 
-	SourceToTransformMap.Add(RightIndexFingerTip_Name, &RightHand.IndexFinger.Tip);
-	SourceToTransformMap.Add(RightIndexFingerDIP_Name, &RightHand.IndexFinger.DIP);
-	SourceToTransformMap.Add(RightIndexFingerPIP_Name, &RightHand.IndexFinger.PIP);
-	SourceToTransformMap.Add(RightIndexFingerMCP_Name, &RightHand.IndexFinger.MCP);
+	SourceToTransformMap.Add(RightIndexFingerTip_Name, &RightHand.IndexFinger.Tip.Transform);
+	SourceToTransformMap.Add(RightIndexFingerDIP_Name, &RightHand.IndexFinger.DIP.Transform);
+	SourceToTransformMap.Add(RightIndexFingerPIP_Name, &RightHand.IndexFinger.PIP.Transform);
+	SourceToTransformMap.Add(RightIndexFingerMCP_Name, &RightHand.IndexFinger.MCP.Transform);
 
-	SourceToTransformMap.Add(RightMiddleFingerTip_Name, &RightHand.MiddleFinger.Tip);
-	SourceToTransformMap.Add(RightMiddleFingerDIP_Name, &RightHand.MiddleFinger.DIP);
-	SourceToTransformMap.Add(RightMiddleFingerPIP_Name, &RightHand.MiddleFinger.PIP);
-	SourceToTransformMap.Add(RightMiddleFingerMCP_Name, &RightHand.MiddleFinger.MCP);
+	SourceToTransformMap.Add(RightMiddleFingerTip_Name, &RightHand.MiddleFinger.Tip.Transform);
+	SourceToTransformMap.Add(RightMiddleFingerDIP_Name, &RightHand.MiddleFinger.DIP.Transform);
+	SourceToTransformMap.Add(RightMiddleFingerPIP_Name, &RightHand.MiddleFinger.PIP.Transform);
+	SourceToTransformMap.Add(RightMiddleFingerMCP_Name, &RightHand.MiddleFinger.MCP.Transform);
 
-	SourceToTransformMap.Add(RightRingFingerTip_Name, &RightHand.RingFinger.Tip);
-	SourceToTransformMap.Add(RightRingFingerDIP_Name, &RightHand.RingFinger.DIP);
-	SourceToTransformMap.Add(RightRingFingerPIP_Name, &RightHand.RingFinger.PIP);
-	SourceToTransformMap.Add(RightRingFingerMCP_Name, &RightHand.RingFinger.MCP);
+	SourceToTransformMap.Add(RightRingFingerTip_Name, &RightHand.RingFinger.Tip.Transform);
+	SourceToTransformMap.Add(RightRingFingerDIP_Name, &RightHand.RingFinger.DIP.Transform);
+	SourceToTransformMap.Add(RightRingFingerPIP_Name, &RightHand.RingFinger.PIP.Transform);
+	SourceToTransformMap.Add(RightRingFingerMCP_Name, &RightHand.RingFinger.MCP.Transform);
 
-	SourceToTransformMap.Add(RightPinkyFingerTip_Name, &RightHand.PinkyFinger.Tip);
-	SourceToTransformMap.Add(RightPinkyFingerDIP_Name, &RightHand.PinkyFinger.DIP);
-	SourceToTransformMap.Add(RightPinkyFingerPIP_Name, &RightHand.PinkyFinger.PIP);
-	SourceToTransformMap.Add(RightPinkyFingerMCP_Name, &RightHand.PinkyFinger.MCP);
+	SourceToTransformMap.Add(RightPinkyFingerTip_Name, &RightHand.PinkyFinger.Tip.Transform);
+	SourceToTransformMap.Add(RightPinkyFingerDIP_Name, &RightHand.PinkyFinger.DIP.Transform);
+	SourceToTransformMap.Add(RightPinkyFingerPIP_Name, &RightHand.PinkyFinger.PIP.Transform);
+	SourceToTransformMap.Add(RightPinkyFingerMCP_Name, &RightHand.PinkyFinger.MCP.Transform);
 
 
 
@@ -498,10 +549,16 @@ FMagicLeapHandTracking::FHandState::FHandState()
 bool FMagicLeapHandTracking::FHandState::GetTransform(EHandTrackingKeypoint KeyPoint, FTransform& OutTransform) const
 {
 	check(static_cast<int32>(KeyPoint) >= 0 && static_cast<int32>(KeyPoint) <= EHandTrackingKeypointCount);
-	OutTransform = *EnumToTransformMap[static_cast<int32>(KeyPoint)];
+	OutTransform = EnumToTransformMap[static_cast<int32>(KeyPoint)]->Transform;
 	return IsValid();
 }
 
+const FMagicLeapHandTracking::FTransformRecord& FMagicLeapHandTracking::FHandState::GetTransform(EHandTrackingKeypoint KeyPoint) const
+{
+	check(static_cast<int32>(KeyPoint) >= 0 && static_cast<int32>(KeyPoint) <= EHandTrackingKeypointCount);
+	return *EnumToTransformMap[static_cast<int32>(KeyPoint)];
+}
+
 const FTransform* FMagicLeapHandTracking::FindTransformBySource(FName SourceName) const
 {
 	return SourceToTransformMap.FindRef(SourceName);
@@ -552,22 +609,22 @@ bool FMagicLeapHandTracking::GetControllerOrientationAndPosition(const int32 Con
 			switch (DeviceHand)
 			{
 			case EControllerHand::Special_1:
-				ControllerTransform = &LeftHand.HandCenter;
+				ControllerTransform = &LeftHand.HandCenter.Transform;
 				break;
 			case EControllerHand::Special_3:
-				ControllerTransform = &LeftHand.IndexFinger.Tip;
+				ControllerTransform = &LeftHand.IndexFinger.Tip.Transform;
 				break;
 			case EControllerHand::Special_5:
-				ControllerTransform = &LeftHand.Thumb.Tip;
+				ControllerTransform = &LeftHand.Thumb.Tip.Transform;
 				break;
 			case EControllerHand::Special_2:
-				ControllerTransform = &RightHand.HandCenter;
+				ControllerTransform = &RightHand.HandCenter.Transform;
 				break;
 			case EControllerHand::Special_4:
-				ControllerTransform = &RightHand.IndexFinger.Tip;
+				ControllerTransform = &RightHand.IndexFinger.Tip.Transform;
 				break;
 			case EControllerHand::Special_6:
-				ControllerTransform = &RightHand.Thumb.Tip;
+				ControllerTransform = &RightHand.Thumb.Tip.Transform;
 				break;
 			default:
 				check(false);
@@ -890,14 +947,15 @@ float FMagicLeapHandTracking::GetGestureConfidenceThreshold(EHandTrackingGesture
 namespace MagicLeapHandTracking
 {
 #if WITH_MLSDK
-	bool FetchTransform(const FAppFramework& AppFramework, const MLKeyPointState& Source, FTransform& OutDest, const TCHAR* DebugString, const TCHAR* DebugString2, const TCHAR* DebugString3)
+	bool FetchTransform(const FAppFramework& AppFramework, const MLKeyPointState& Source, FMagicLeapHandTracking::FTransformRecord& OutDest, const TCHAR* DebugString, const TCHAR* DebugString2, const TCHAR* DebugString3)
 	{
 		bool bResult = false;
 
 		if (Source.is_valid)
 		{
 			EFailReason FailReason = EFailReason::None;
-			bResult = AppFramework.GetTransform(Source.frame_id, OutDest, FailReason);
+			FTransform Transform;
+			bResult = AppFramework.GetTransform(Source.frame_id, Transform, FailReason);
 			if (!bResult)
 			{
 				if (FailReason == EFailReason::NaNsInTransform)
@@ -905,6 +963,11 @@ namespace MagicLeapHandTracking
 					UE_LOG(LogMagicLeapHandTracking, Error, TEXT("NaNs in %s %s %s transform."), DebugString, DebugString2, DebugString3);
 				}
 			}
+			else
+			{
+				OutDest.Transform = Transform;
+				OutDest.bWritten = true;
+			}
 
 		}
 		return bResult;
@@ -1072,6 +1135,8 @@ void FMagicLeapHandTracking::UpdateTrackerData()
 
 		LeftHand.GestureConfidence = HandTrackingData.left_hand_state.keypose_confidence[HandTrackingData.left_hand_state.keypose];
 		RightHand.GestureConfidence = HandTrackingData.right_hand_state.keypose_confidence[HandTrackingData.right_hand_state.keypose];
+
+		UpdateLiveLink();
 	}
 #endif //WITH_MLSDK
 }
diff --git a/Engine/Plugins/Lumin/MagicLeap/Source/MagicLeapHandTracking/Private/MagicLeapHandTracking.h b/Engine/Plugins/Lumin/MagicLeap/Source/MagicLeapHandTracking/Private/MagicLeapHandTracking.h
--- a/Engine/Plugins/Lumin/MagicLeap/Source/MagicLeapHandTracking/Private/MagicLeapHandTracking.h
+++ b/Engine/Plugins/Lumin/MagicLeap/Source/MagicLeapHandTracking/Private/MagicLeapHandTracking.h
@@ -17,13 +17,46 @@ ML_INCLUDES_END
 #include "XRMotionControllerBase.h"
 #include "InputCoreTypes.h"
 #include "AppEventHandler.h"
+#include "ILiveLinkSource.h"
 
 /**
   * MagicLeap HandTracking
   */
-class FMagicLeapHandTracking : public IMagicLeapInputDevice, public FXRMotionControllerBase, public MagicLeap::IAppEventHandler
+class FMagicLeapHandTracking : public IMagicLeapInputDevice, public FXRMotionControllerBase, public MagicLeap::IAppEventHandler, public ILiveLinkSource
 {
 public:
+	static FName HandCenter_Name;
+
+	static FName WristCenter_Name;
+	static FName WristUlnar_Name;
+	static FName WristRadial_Name;
+
+	static FName ThumbTip_Name;
+	static FName ThumbIP_Name;
+	static FName ThumbMCP_Name;
+	static FName ThumbCMC_Name;
+
+	static FName IndexFingerTip_Name;
+	static FName IndexFingerDIP_Name;
+	static FName IndexFingerPIP_Name;
+	static FName IndexFingerMCP_Name;
+
+	static FName MiddleFingerTip_Name;
+	static FName MiddleFingerDIP_Name;
+	static FName MiddleFingerPIP_Name;
+	static FName MiddleFingerMCP_Name;
+
+	static FName RingFingerTip_Name;
+	static FName RingFingerDIP_Name;
+	static FName RingFingerPIP_Name;
+	static FName RingFingerMCP_Name;
+
+	static FName PinkyFingerTip_Name;
+	static FName PinkyFingerDIP_Name;
+	static FName PinkyFingerPIP_Name;
+	static FName PinkyFingerMCP_Name;
+
+
 	static FName LeftHandCenter_Name;
 
 	static FName LeftWristCenter_Name;
@@ -87,26 +120,27 @@ public:
 	static FName RightPinkyFingerPIP_Name;
 	static FName RightPinkyFingerMCP_Name;
 
-	//struct FTransformState
-	//{
-
-	//};
+	struct FTransformRecord
+	{
+		FTransform Transform = FTransform::Identity;
+		bool bWritten = false;
+	};
 	struct FWristTransforms : public FNoncopyable
 	{
 		// Wrist center.
-		FTransform Center = FTransform::Identity;
+		FTransformRecord Center;
 		// Ulnar-sided wrist
-		FTransform Ulnar = FTransform::Identity;
+		FTransformRecord Ulnar;
 		// Radial-sided wrist.
-		FTransform Radial = FTransform::Identity;
+		FTransformRecord Radial;
 	};
 	struct FDigitTransforms : public FNoncopyable
 	{
 		// These labels are not correct anatomical nomenclature for the thumb, but they 1:1 map with the correct names.
-		FTransform Tip = FTransform::Identity; // digit tip
-		FTransform DIP = FTransform::Identity; // distal joint
-		FTransform PIP = FTransform::Identity; // proximal joint
-		FTransform MCP = FTransform::Identity; // base of digit
+		FTransformRecord Tip; // digit tip
+		FTransformRecord DIP; // distal joint
+		FTransformRecord PIP; // proximal joint
+		FTransformRecord MCP; // base of digit
 	};
 	struct FHandState : public FNoncopyable
 	{
@@ -124,13 +158,14 @@ public:
 
 		FWristTransforms Wrist;
 
-		FTransform HandCenter = FTransform::Identity;
+		FTransformRecord HandCenter;
 
 		bool IsValid() const { return Gesture != EHandTrackingGesture::NoHand; }
 		bool GetTransform(EHandTrackingKeypoint KeyPoint, FTransform& OutTransform) const;
+		const FTransformRecord& GetTransform(EHandTrackingKeypoint KeyPoint) const;
 
 	private:
-		FTransform* EnumToTransformMap[EHandTrackingKeypointCount];
+		FTransformRecord* EnumToTransformMap[EHandTrackingKeypointCount];
 	};
 
 public:
@@ -144,6 +179,14 @@ public:
 	virtual FName GetMotionControllerDeviceTypeName() const override;
 	virtual void EnumerateSources(TArray<FMotionControllerSource>& SourcesOut) const override;
 
+	// ILiveLinkSource interface
+	virtual void ReceiveClient(ILiveLinkClient* InClient, FGuid InSourceGuid) override;
+	virtual bool IsSourceStillValid() override;
+	virtual bool RequestSourceShutdown() override;
+	virtual FText GetSourceMachineName() const override;
+	virtual FText GetSourceStatus() const override;
+	virtual FText GetSourceType() const override;
+	// End ILiveLinkSource
 
 	/** IMagicLeapInputDevice interface */
 	virtual void Tick(float DeltaTime) override;
@@ -179,6 +222,10 @@ private:
 	void OnAppPause() override;
 	void OnAppResume() override;
 
+	void SetupLiveLinkData();
+	void UpdateLiveLink();
+	void UpdateLiveLinkTransforms(TArray<FTransform>& OutTransforms, const FMagicLeapHandTracking::FHandState& HandState);
+
 #if WITH_MLSDK
 	const MLHandTrackingData& GetCurrentHandTrackingData()				{ return HandTrackingDatas[CurrentHandTrackingDataIndex]; }
 	const MLHandTrackingData& GetPreviousHandTrackingData()				{ return HandTrackingDatas[1 - CurrentHandTrackingDataIndex]; }
@@ -197,6 +244,8 @@ private:
 	MLHandTrackingData HandTrackingDatas[2];
 	int32 CurrentHandTrackingDataIndex = 0;
 	MLHandTrackingStaticData HandTrackingStaticData;
+	TArray<int32> BoneParents;
+	TArray<EHandTrackingKeypoint> BoneKeypoints;
 #endif //WITH_MLSDK
 
 	FHandState LeftHand;
@@ -236,6 +285,18 @@ private:
 		static const FKey Right_NoPose;
 		static const FKey Right_NoHand;
 	};
+
+	// LiveLink Data
+	/** The local client to push data updates to */
+	ILiveLinkClient* LiveLinkClient = nullptr;
+	/** Our identifier in LiveLink */
+	FGuid LiveLinkSourceGuid;
+
+	static FName LiveLinkLeftHandTrackingSubjectName;
+	static FName LiveLinkRightHandTrackingSubjectName;
+	bool bNewLiveLinkClient = false;
+	FLiveLinkRefSkeleton LiveLinkRefSkeleton;
+
 };
 
 DEFINE_LOG_CATEGORY_STATIC(LogMagicLeapHandTracking, Display, All);
diff --git a/Engine/Plugins/Lumin/MagicLeap/Source/MagicLeapHandTracking/Private/MagicLeapHandTrackingFunctionLibrary.cpp b/Engine/Plugins/Lumin/MagicLeap/Source/MagicLeapHandTracking/Private/MagicLeapHandTrackingFunctionLibrary.cpp
--- a/Engine/Plugins/Lumin/MagicLeap/Source/MagicLeapHandTracking/Private/MagicLeapHandTrackingFunctionLibrary.cpp
+++ b/Engine/Plugins/Lumin/MagicLeap/Source/MagicLeapHandTracking/Private/MagicLeapHandTrackingFunctionLibrary.cpp
@@ -6,6 +6,8 @@
 #include "IMagicLeapPlugin.h"
 #include "HeadMountedDisplayFunctionLibrary.h"
 #include "Engine/World.h"
+#include "ILiveLinkClient.h"
+#include "ILiveLinkSource.h"
 
 bool UMagicLeapHandTrackingFunctionLibrary::GetHandCenter(EControllerHand Hand, FTransform& HandCenter)
 {
@@ -16,7 +18,7 @@ bool UMagicLeapHandTrackingFunctionLibrary::GetHandCenter(EControllerHand Hand,
 		if (Hand == EControllerHand::Left || Hand == EControllerHand::Right)
 		{
 			const FMagicLeapHandTracking::FHandState& HandTrackingData = (Hand == EControllerHand::Right) ? HandTracking->GetRightHandState() : HandTracking->GetLeftHandState();
-			HandCenter = HandTrackingData.HandCenter;
+			HandCenter = HandTrackingData.HandCenter.Transform;
 			return HandTrackingData.IsValid();
 		}
 		else
@@ -42,18 +44,18 @@ bool UMagicLeapHandTrackingFunctionLibrary::GetHandIndexFingerTip(EControllerHan
 			{
 			case EGestureTransformSpace::Tracking:
 			{
-				Pointer = HandTrackingData.IndexFinger.Tip;
+				Pointer = HandTrackingData.IndexFinger.Tip.Transform;
 				break;
 			}
 			case EGestureTransformSpace::World:
 			{
 				FTransform TrackingToWorldTransform = UHeadMountedDisplayFunctionLibrary::GetTrackingToWorldTransform(GWorld);
-				Pointer = HandTrackingData.IndexFinger.Tip * TrackingToWorldTransform;
+				Pointer = HandTrackingData.IndexFinger.Tip.Transform * TrackingToWorldTransform;
 				break;
 			}
 			case EGestureTransformSpace::Hand:
 			{
-				Pointer = HandTrackingData.IndexFinger.Tip * HandTrackingData.HandCenter.Inverse();
+				Pointer = HandTrackingData.IndexFinger.Tip.Transform * HandTrackingData.HandCenter.Transform.Inverse();
 				break;
 			}
 			default:
@@ -85,18 +87,18 @@ bool UMagicLeapHandTrackingFunctionLibrary::GetHandThumbTip(EControllerHand Hand
 			{
 			case EGestureTransformSpace::Tracking:
 			{
-				Secondary = HandTrackingData.Thumb.Tip;
+				Secondary = HandTrackingData.Thumb.Tip.Transform;
 				break;
 			}
 			case EGestureTransformSpace::World:
 			{
 				FTransform TrackingToWorldTransform = UHeadMountedDisplayFunctionLibrary::GetTrackingToWorldTransform(GWorld);
-				Secondary = HandTrackingData.Thumb.Tip * TrackingToWorldTransform;
+				Secondary = HandTrackingData.Thumb.Tip.Transform * TrackingToWorldTransform;
 				break;
 			}
 			case EGestureTransformSpace::Hand:
 			{
-				Secondary = HandTrackingData.Thumb.Tip * HandTrackingData.HandCenter.Inverse();
+				Secondary = HandTrackingData.Thumb.Tip.Transform * HandTrackingData.HandCenter.Transform.Inverse();
 				break;
 			}
 			default:
@@ -157,9 +159,9 @@ bool UMagicLeapHandTrackingFunctionLibrary::GetGestureKeypoints(EControllerHand
 		{
 			const FMagicLeapHandTracking::FHandState& HandState = (Hand == EControllerHand::Left) ? HandTracking->GetLeftHandState() : HandTracking->GetRightHandState();
 			Keypoints.SetNum(3);
-			Keypoints[0] = HandState.HandCenter;
-			Keypoints[1] = HandState.IndexFinger.Tip;
-			Keypoints[2] = HandState.Thumb.Tip;
+			Keypoints[0] = HandState.HandCenter.Transform;
+			Keypoints[1] = HandState.IndexFinger.Tip.Transform;
+			Keypoints[2] = HandState.Thumb.Tip.Transform;
 			return true;
 		}
 		else
@@ -297,3 +299,22 @@ bool UMagicLeapHandTrackingFunctionLibrary::GetCurrentGesture(EControllerHand Ha
 	Gesture = EHandTrackingGesture::NoHand;
 	return false;
 }
+
+bool UMagicLeapHandTrackingFunctionLibrary::GetMagicLeapHandTrackingLiveLinkSource(FLiveLinkSourceHandle& SourceHandle)
+{
+	IModularFeatures& ModularFeatures = IModularFeatures::Get();
+
+	if (ModularFeatures.IsModularFeatureAvailable(ILiveLinkClient::ModularFeatureName))
+	{
+		TSharedPtr<ILiveLinkSource> HandTrackingSource = StaticCastSharedPtr<FMagicLeapHandTracking>(IMagicLeapHandTrackingPlugin::Get().GetLiveLinkSource());
+		ILiveLinkClient* LiveLinkClient = &ModularFeatures.GetModularFeature<ILiveLinkClient>(ILiveLinkClient::ModularFeatureName);
+		LiveLinkClient->AddSource(HandTrackingSource);
+		SourceHandle.SetSourcePointer(HandTrackingSource);
+		return true;
+	}
+	else
+	{
+		SourceHandle.SetSourcePointer(nullptr);
+		return false;
+	}
+};
diff --git a/Engine/Plugins/Lumin/MagicLeap/Source/MagicLeapHandTracking/Private/MagicLeapHandTrackingLiveLink.cpp b/Engine/Plugins/Lumin/MagicLeap/Source/MagicLeapHandTracking/Private/MagicLeapHandTrackingLiveLink.cpp
new file mode 100644
--- /dev/null
+++ b/Engine/Plugins/Lumin/MagicLeap/Source/MagicLeapHandTracking/Private/MagicLeapHandTrackingLiveLink.cpp
@@ -0,0 +1,187 @@
+// Copyright 1998-2018 Epic Games, Inc. All Rights Reserved.
+
+#include "MagicLeapHandTracking.h"
+#include "IMagicLeapHandTrackingPlugin.h"
+#include "LiveLinkMagicLeapHandTrackingSourceEditor.h"
+#include "CoreMinimal.h"
+#include "Engine/Engine.h"
+#include "LiveLinkSourceFactory.h"
+#include "ILiveLinkClient.h"
+
+#define LOCTEXT_NAMESPACE "MagicLeapHandTracking"
+
+void FMagicLeapHandTracking::ReceiveClient(ILiveLinkClient* InClient, FGuid InSourceGuid)
+{
+	LiveLinkClient = InClient;
+	LiveLinkSourceGuid = InSourceGuid;
+	bNewLiveLinkClient = true;
+}
+
+bool FMagicLeapHandTracking::IsSourceStillValid()
+{
+	return LiveLinkClient != nullptr;
+}
+
+bool FMagicLeapHandTracking::RequestSourceShutdown()
+{
+	LiveLinkClient = nullptr;
+	LiveLinkSourceGuid.Invalidate();
+	return true;
+}
+
+FText FMagicLeapHandTracking::GetSourceMachineName() const
+{
+	return FText().FromString(FPlatformProcess::ComputerName());
+}
+
+FText FMagicLeapHandTracking::GetSourceStatus() const
+{
+	return LOCTEXT("MagicLeapHandTrackingLiveLinkStatus", "Active");
+}
+
+FText FMagicLeapHandTracking::GetSourceType() const
+{
+	return LOCTEXT("MagicLeapHandTrackingLiveLinkSourceType", "MagicLeap Hand Tracking");
+}
+
+#define MLHTBONE(keyenum, name, parent) BoneKeypoints.Add(keyenum); BoneNames.Add(name); BoneParents.Add(parent);
+void FMagicLeapHandTracking::SetupLiveLinkData()
+{
+	TArray<FName> BoneNames;
+	BoneNames.Reserve(EHandTrackingKeypointCount);
+
+	// Array of bone indices to parent bone index
+	BoneParents.Reserve(EHandTrackingKeypointCount);
+
+	BoneKeypoints.Reserve(EHandTrackingKeypointCount);
+
+	// Only one can be the root (-1 parent)!
+
+	MLHTBONE(EHandTrackingKeypoint::Thumb_Tip, ThumbTip_Name,			1);
+	MLHTBONE(EHandTrackingKeypoint::Thumb_IP,	ThumbIP_Name,			2);
+	MLHTBONE(EHandTrackingKeypoint::Thumb_MCP, ThumbMCP_Name,			3);
+	MLHTBONE(EHandTrackingKeypoint::Thumb_CMC, ThumbCMC_Name,			23); //3
+
+	MLHTBONE(EHandTrackingKeypoint::Index_Tip, IndexFingerTip_Name,		5);
+	MLHTBONE(EHandTrackingKeypoint::Index_DIP, IndexFingerDIP_Name,		6);
+	MLHTBONE(EHandTrackingKeypoint::Index_PIP, IndexFingerPIP_Name,		7);
+	MLHTBONE(EHandTrackingKeypoint::Index_MCP, IndexFingerMCP_Name,		23); //7
+
+	MLHTBONE(EHandTrackingKeypoint::Middle_Tip, MiddleFingerTip_Name,	9);
+	MLHTBONE(EHandTrackingKeypoint::Middle_DIP, MiddleFingerDIP_Name,	10);
+	MLHTBONE(EHandTrackingKeypoint::Middle_PIP, MiddleFingerPIP_Name,	11);
+	MLHTBONE(EHandTrackingKeypoint::Middle_MCP, MiddleFingerMCP_Name,	23); //11
+
+	MLHTBONE(EHandTrackingKeypoint::Ring_Tip, RingFingerTip_Name,		13);
+	MLHTBONE(EHandTrackingKeypoint::Ring_DIP, RingFingerDIP_Name,		14);
+	MLHTBONE(EHandTrackingKeypoint::Ring_PIP, RingFingerPIP_Name,		15);
+	MLHTBONE(EHandTrackingKeypoint::Ring_MCP, RingFingerMCP_Name,		23); //15
+
+	MLHTBONE(EHandTrackingKeypoint::Pinky_Tip, PinkyFingerTip_Name,		17);
+	MLHTBONE(EHandTrackingKeypoint::Pinky_DIP, PinkyFingerDIP_Name,		18);
+	MLHTBONE(EHandTrackingKeypoint::Pinky_PIP, PinkyFingerPIP_Name,		19);
+	MLHTBONE(EHandTrackingKeypoint::Pinky_MCP, PinkyFingerMCP_Name,		23); //19
+
+	MLHTBONE(EHandTrackingKeypoint::Wrist_Center,	WristCenter_Name,	23);
+	MLHTBONE(EHandTrackingKeypoint::Wrist_Ulnar,	WristUlnar_Name,	23);
+	MLHTBONE(EHandTrackingKeypoint::Wrist_Radial,	WristRadial_Name,	23);
+
+	MLHTBONE(EHandTrackingKeypoint::Hand_Center, HandCenter_Name,		-1); //23 //Root
+
+	LiveLinkRefSkeleton.SetBoneNames(BoneNames);
+	LiveLinkRefSkeleton.SetBoneParents(BoneParents);
+}
+#undef MLHTBONE
+
+void FMagicLeapHandTracking::UpdateLiveLinkTransforms(TArray<FTransform>& OutTransforms, const FMagicLeapHandTracking::FHandState& HandState)
+{
+	// Live link transforms need to be in the hierarchical skeleton, so each in the space of its parent.
+	// The hand tracking transforms are in world space.
+	// Sadly hand tracking transforms can be unused, and contain only an identity matrix.  So we do some weird stuff to keep those
+	// identity in the hierarchical skeleton.
+
+	for (int32 i = 0; i < EHandTrackingKeypointCount; ++i)
+	{
+		const FTransformRecord& TransformRecord = HandState.GetTransform(BoneKeypoints[i]);
+		bool bWritten = TransformRecord.bWritten;
+		if (bWritten)
+		{
+			int32 ParentIndex = BoneParents[i];
+			if (ParentIndex < 0)
+			{
+				// We are at the root, so use it.  (even if it is identity)
+				OutTransforms[i] = TransformRecord.Transform;
+			}
+			else
+			{
+				const FTransformRecord* PossibleParentTransformRecord = &(HandState.GetTransform(BoneKeypoints[ParentIndex]));
+				while (!PossibleParentTransformRecord->bWritten && ParentIndex >= 0)
+				{
+					PossibleParentTransformRecord = &(HandState.GetTransform(BoneKeypoints[ParentIndex]));
+					ParentIndex = BoneParents[ParentIndex];
+				}
+				if (PossibleParentTransformRecord->bWritten)
+				{
+					// We found a parent with a transform, so we will be relative to that.
+					// Any intermediates are identity.
+					OutTransforms[i] = TransformRecord.Transform * PossibleParentTransformRecord->Transform.Inverse();
+				}
+				else
+				{
+					check(ParentIndex < 0);
+					// We are at the root, but it is unwritten, so go ahead and use the original transform.
+					// Everything above it was unwritten, and identity.
+					OutTransforms[i] = TransformRecord.Transform;
+				}
+			}
+		}
+		else
+		{
+			// if this transform is not written leave it as is.
+		}
+	}
+}
+
+void FMagicLeapHandTracking::UpdateLiveLink()
+{
+	check(IsInGameThread());
+
+	if (LiveLinkClient)
+	{
+		if (bNewLiveLinkClient)
+		{
+			LiveLinkClient->ClearSubject(LiveLinkLeftHandTrackingSubjectName);
+			LiveLinkClient->ClearSubject(LiveLinkRightHandTrackingSubjectName);
+			LiveLinkClient->PushSubjectSkeleton(LiveLinkSourceGuid, LiveLinkLeftHandTrackingSubjectName, LiveLinkRefSkeleton);
+			LiveLinkClient->PushSubjectSkeleton(LiveLinkSourceGuid, LiveLinkRightHandTrackingSubjectName, LiveLinkRefSkeleton);
+			bNewLiveLinkClient = false;
+		}
+
+		static FLiveLinkFrameData LiveLinkLeftFrame;
+		static FLiveLinkFrameData LiveLinkRightFrame;
+		static bool bInitialized = false;
+		if (!bInitialized)
+		{
+			LiveLinkLeftFrame.Transforms.Reserve(EHandTrackingKeypointCount);
+			LiveLinkRightFrame.Transforms.Reserve(EHandTrackingKeypointCount);
+			for (size_t i = 0; i < EHandTrackingKeypointCount; ++i)
+			{
+				LiveLinkLeftFrame.Transforms.Add(FTransform::Identity);
+				LiveLinkRightFrame.Transforms.Add(FTransform::Identity);
+			}
+			bInitialized = true;
+		}
+
+		LiveLinkLeftFrame.WorldTime = LiveLinkRightFrame.WorldTime = FPlatformTime::Seconds();
+
+		// Update the transforms for each subject from tracking data
+		UpdateLiveLinkTransforms(LiveLinkLeftFrame.Transforms, LeftHand);
+		UpdateLiveLinkTransforms(LiveLinkRightFrame.Transforms, RightHand);
+
+		// Share the data locally with the LiveLink client
+		LiveLinkClient->PushSubjectData(LiveLinkSourceGuid, LiveLinkLeftHandTrackingSubjectName, LiveLinkLeftFrame);
+		LiveLinkClient->PushSubjectData(LiveLinkSourceGuid, LiveLinkRightHandTrackingSubjectName, LiveLinkRightFrame);
+	}
+}
+
+#undef LOCTEXT_NAMESPACE
diff --git a/Engine/Plugins/Lumin/MagicLeap/Source/MagicLeapHandTracking/Public/IMagicLeapHandTrackingPlugin.h b/Engine/Plugins/Lumin/MagicLeap/Source/MagicLeapHandTracking/Public/IMagicLeapHandTrackingPlugin.h
--- a/Engine/Plugins/Lumin/MagicLeap/Source/MagicLeapHandTracking/Public/IMagicLeapHandTrackingPlugin.h
+++ b/Engine/Plugins/Lumin/MagicLeap/Source/MagicLeapHandTracking/Public/IMagicLeapHandTrackingPlugin.h
@@ -40,4 +40,11 @@ public:
 	 * @return Shared pointer to the Input Device.
 	 */
 	virtual TSharedPtr<IInputDevice> GetInputDevice() = 0;
+
+	/**
+	* Returns the LiveLinkSource associated with this MagicLeapHandTrackingPlugin.
+	*
+	* @return Shared pointer to the Input Device.
+	*/
+	virtual TSharedPtr<class ILiveLinkSource> GetLiveLinkSource() = 0;
 };
diff --git a/Engine/Plugins/Lumin/MagicLeap/Source/MagicLeapHandTracking/Public/MagicLeapHandTrackingFunctionLibrary.h b/Engine/Plugins/Lumin/MagicLeap/Source/MagicLeapHandTracking/Public/MagicLeapHandTrackingFunctionLibrary.h
--- a/Engine/Plugins/Lumin/MagicLeap/Source/MagicLeapHandTracking/Public/MagicLeapHandTrackingFunctionLibrary.h
+++ b/Engine/Plugins/Lumin/MagicLeap/Source/MagicLeapHandTracking/Public/MagicLeapHandTrackingFunctionLibrary.h
@@ -7,6 +7,7 @@
 #include "UObject/ObjectMacros.h"
 #include "Kismet/BlueprintFunctionLibrary.h"
 #include "InputCoreTypes.h"
+#include "ILiveLinkSource.h"
 #include "MagicLeapHandTrackingFunctionLibrary.generated.h"
 
 UCLASS(ClassGroup = MagicLeap)
@@ -155,4 +156,15 @@ public:
 	*/
 	UFUNCTION(BlueprintCallable, Category = "HandTracking|MagicLeap")
 	static bool GetCurrentGesture(EControllerHand Hand, EHandTrackingGesture& Gesture);
+
+
+	/**
+	Get a LiveLinkSourceHandle for magic leap hand tracking.
+
+	@param SourceHandle Output parameter SourceHandle that will be filled in.
+	@return true if a LiveLink Source was assigned.
+	*/
+	UFUNCTION(BlueprintCallable, Category = "HandTracking|MagicLeap")
+	static bool GetMagicLeapHandTrackingLiveLinkSource(struct FLiveLinkSourceHandle& SourceHandle);
+
 };
diff --git a/Engine/Source/Runtime/LiveLinkInterface/Public/ILiveLinkSource.h b/Engine/Source/Runtime/LiveLinkInterface/Public/ILiveLinkSource.h
--- a/Engine/Source/Runtime/LiveLinkInterface/Public/ILiveLinkSource.h
+++ b/Engine/Source/Runtime/LiveLinkInterface/Public/ILiveLinkSource.h
@@ -3,8 +3,10 @@
 #pragma once
 
 #include "CoreMinimal.h"
+#include "UObject/ObjectMacros.h"
 #include "Misc/Guid.h"
 #include "LiveLinkTypes.h"
+#include "ILiveLinkSource.generated.h"
 
 class ILiveLinkClient;
 class ULiveLinkSourceSettings;
@@ -30,4 +32,22 @@ public:
 
 	virtual UClass* GetCustomSettingsClass() const { return nullptr; }
 	virtual void OnSettingsChanged(ULiveLinkSourceSettings* Settings, const FPropertyChangedEvent& PropertyChangedEvent) {}
-};
\ No newline at end of file
+};
+
+// A Blueprint handle to a specific LiveLink Source
+USTRUCT(BlueprintType)
+struct FLiveLinkSourceHandle
+{
+	GENERATED_USTRUCT_BODY()
+
+	FLiveLinkSourceHandle() = default;
+
+	virtual ~FLiveLinkSourceHandle() = default;
+
+	void SetSourcePointer(TSharedPtr<ILiveLinkSource> InSourcePointer)
+	{
+		SourcePointer = InSourcePointer;
+	};
+
+	TSharedPtr<ILiveLinkSource> SourcePointer;
+};
