#DEV-VR - Saving off original MLSDK variable for future editor "restarts" within the same process

diff --git a/Engine/Plugins/Lumin/MagicLeap/Source/MLSDK/Public/MagicLeapPluginUtil.h b/Engine/Plugins/Lumin/MagicLeap/Source/MLSDK/Public/MagicLeapPluginUtil.h
--- a/Engine/Plugins/Lumin/MagicLeap/Source/MLSDK/Public/MagicLeapPluginUtil.h
+++ b/Engine/Plugins/Lumin/MagicLeap/Source/MLSDK/Public/MagicLeapPluginUtil.h
@@ -32,6 +32,21 @@ public:
 
 		// Give preference to the config setting.
 		FString MLSDKPath;
+		FString MLSDKEnvironmentVariableName = TEXT("MLSDK");
+		FString OriginalMLSDKEnvironmentVariableName = TEXT("MLSDKOriginal");
+
+		//save off the value of MLSDK, if it hasn't been saved before
+		FPlatformMisc::GetEnvironmentVariable(*OriginalMLSDKEnvironmentVariableName, MLSDK.GetCharArray().GetData(), 1024);
+		MLSDK.TrimToNullTerminator();
+		if (MLSDK.IsEmpty())
+		{
+			//only once, read the starting value
+			MLSDK = FString::ChrN(1024, TEXT('\0'));
+			FPlatformMisc::GetEnvironmentVariable(*MLSDKEnvironmentVariableName, MLSDK.GetCharArray().GetData(), 1024);
+			FPlatformMisc::SetEnvironmentVar(*OriginalMLSDKEnvironmentVariableName, *MLSDK);
+		}
+
+
 		GConfig->GetString(TEXT("/Script/LuminPlatformEditor.MagicLeapSDKSettings"), TEXT("MLSDKPath"), MLSDKPath, GEngineIni);
 		MLSDKPath.TrimToNullTerminator();
 		if (MLSDKPath.Len() > 0)
@@ -52,12 +67,12 @@ public:
 
 			if (MLSDKPath.Len() > 0 && FPaths::DirectoryExists(MLSDKPath))
 			{
-				FPlatformMisc::SetEnvironmentVar(TEXT("MLSDK"), *MLSDKPath);
+				MLSDK = MLSDKPath;
 			}
 		}
 
-		FPlatformMisc::GetEnvironmentVariable(TEXT("MLSDK"), MLSDK.GetCharArray().GetData(), 1024);
 		MLSDK.TrimToNullTerminator();
+		FPlatformMisc::SetEnvironmentVar(*MLSDKEnvironmentVariableName, *MLSDK);
 
 		if (CheckForVDZILibraries)
 		{
@@ -109,8 +124,19 @@ public:
 		// to work in the plugin.
 		const int32 MaxPathVarLen = 32768;
 		FString PathVar = FString::ChrN(MaxPathVarLen, TEXT('\0'));
-		FPlatformMisc::GetEnvironmentVariable(TEXT("PATH"), PathVar.GetCharArray().GetData(), MaxPathVarLen);
+
+		//if we've previously saved the original path off, just use that saved original version
+		FPlatformMisc::GetEnvironmentVariable(TEXT("PATHOriginal"), PathVar.GetCharArray().GetData(), MaxPathVarLen);
 		PathVar.TrimToNullTerminator();
+		if (PathVar.IsEmpty())
+		{
+			//save off the path before we add to it
+			PathVar = FString::ChrN(MaxPathVarLen, TEXT('\0'));
+			FPlatformMisc::GetEnvironmentVariable(TEXT("PATH"), PathVar.GetCharArray().GetData(), MaxPathVarLen);
+			PathVar.TrimToNullTerminator();
+			FPlatformMisc::SetEnvironmentVar(TEXT("PATHOriginal"), *PathVar);
+		}
+
 		for (const FString& path : DllSearchPaths)
 		{
 			PathVar.Append(FPlatformMisc::GetPathVarDelimiter());
