Disable ML HMD creation when built without the SDK.

diff --git a/Engine/Plugins/Lumin/MagicLeap/Source/MagicLeap/Private/MagicLeapHMD.cpp b/Engine/Plugins/Lumin/MagicLeap/Source/MagicLeap/Private/MagicLeapHMD.cpp
--- a/Engine/Plugins/Lumin/MagicLeap/Source/MagicLeap/Private/MagicLeapHMD.cpp
+++ b/Engine/Plugins/Lumin/MagicLeap/Source/MagicLeap/Private/MagicLeapHMD.cpp
@@ -171,6 +171,7 @@ public:
 		{
 			LocalHMD = HMD.Pin();
 		}
+#if WITH_MLSDK
 		else
 		{
 			//initialize AR system
@@ -182,7 +183,8 @@ public:
 			ARModule->ConnectARImplementationToXRSystem(LocalHMD.Get());
 			ARModule->InitializeARImplementation();
 		}
-#endif
+#endif // WITH_MLSDK
+#endif // !PLATFORM_MAC
 		if (LocalHMD.IsValid() && !LocalHMD->IsInitialized())
 		{
 			LocalHMD->Startup();
