Disable Lumin plugin by default as a temporary workaround

diff --git a/Engine/Plugins/Lumin/MagicLeap/MagicLeap.uplugin b/Engine/Plugins/Lumin/MagicLeap/MagicLeap.uplugin
--- a/Engine/Plugins/Lumin/MagicLeap/MagicLeap.uplugin
+++ b/Engine/Plugins/Lumin/MagicLeap/MagicLeap.uplugin
@@ -7,7 +7,7 @@
 	"CreatedByURL" : "http://magicleap.com",
 	"Description" : "Support for the Magic Leap Head Mounted Display.",
 	"Category" : "Augmented Reality",
-	"EnabledByDefault" : true,
+	"EnabledByDefault" : false,
 	"DocsURL" : "",
 	"MarketplaceURL" : "",
 	"SupportURL" : "",
