Several fixes for Magic Leap being enabled by default when no MLSDK is installed.  Check on startup if VDZI is enabled, and only then attempt to initialize the HMD if it is, and the remoting DLLs can be loaded.  Added early outs and checks to prevent crashes when the MLSDK is not present or removed after compilation.

diff --git a/Engine/Plugins/Lumin/MagicLeap/Source/MagicLeap/Private/MagicLeapHMD.cpp b/Engine/Plugins/Lumin/MagicLeap/Source/MagicLeap/Private/MagicLeapHMD.cpp
--- a/Engine/Plugins/Lumin/MagicLeap/Source/MagicLeap/Private/MagicLeapHMD.cpp
+++ b/Engine/Plugins/Lumin/MagicLeap/Source/MagicLeap/Private/MagicLeapHMD.cpp
@@ -110,7 +110,14 @@ public:
 		{
 #if PLATFORM_WINDOWS
 
-			APISetup.LoadDLL(TEXT("ml_remote"));
+			const bool bLoadedRemoteDLL = APISetup.LoadDLL(TEXT("ml_remote"));
+			if (!bLoadedRemoteDLL)
+			{
+				// Bail early, because we'll eventually die later.
+				UE_LOG(LogMagicLeap, Warning, TEXT("VDZI enabled, but can't load the ml_remote DLL.  Is your MLSDK directory set up properly?"));
+				bIsVDZIEnabled = false;
+			}
+			
 			FString CommandLine = FCommandLine::Get();
 			const FString GLFlag(" -opengl4 ");
 			const FString VKFlag(" -vulkan ");
@@ -165,7 +172,16 @@ public:
 	/** IHeadMountedDisplayModule implementation */
 	virtual TSharedPtr< class IXRTrackingSystem, ESPMode::ThreadSafe > CreateTrackingSystem() override
 	{
+#if !PLATFORM_LUMIN
+		// Early out if VDZI is not enabled on non-Lumin platforms.  We don't want it reporting as available when no MLSDK is present
+		if (!bIsVDZIEnabled)
+		{
+			return nullptr;
+		}
+#endif //PLATFORM_LUMIN
+		
 		TSharedPtr<FMagicLeapHMD, ESPMode::ThreadSafe> LocalHMD;
+
 #if !PLATFORM_MAC
 		if (HMD.IsValid())
 		{
@@ -206,7 +222,7 @@ public:
 		{
 			TSharedPtr<FMagicLeapHMD, ESPMode::ThreadSafe> hmd = StaticCastSharedPtr<FMagicLeapHMD>(GEngine->XRSystem);
 			// IsHMDConnected() is an expensive call when MLRemote is enabled, so we'll keep the onus of that check on the caller.
-			return hmd.IsValid();
+			return hmd.IsValid() && hmd->IsVDZIEnabled();
 		}
 		return false;
 #endif // PLATFORM_LUMIN
@@ -230,6 +246,14 @@ public:
 
 	virtual TSharedPtr<IHeadMountedDisplayVulkanExtensions, ESPMode::ThreadSafe> GetVulkanExtensions() override
 	{
+#if !PLATFORM_LUMIN
+		// Check to see if VDZI is enabled, and abort if not.  We shouldn't modify the extensions if we're not active
+		if (!bIsVDZIEnabled)
+		{
+			return nullptr;
+		}
+#endif //PLATFORM_LUMIN
+
 #if PLATFORM_WINDOWS || PLATFORM_LUMIN
 		if (!VulkanExtensions.IsValid())
 		{
diff --git a/Engine/Plugins/Lumin/MagicLeap/Source/MagicLeap/Private/MagicLeapHMD.h b/Engine/Plugins/Lumin/MagicLeap/Source/MagicLeap/Private/MagicLeapHMD.h
--- a/Engine/Plugins/Lumin/MagicLeap/Source/MagicLeap/Private/MagicLeapHMD.h
+++ b/Engine/Plugins/Lumin/MagicLeap/Source/MagicLeap/Private/MagicLeapHMD.h
@@ -210,6 +210,8 @@ public:
 
 	bool IsPerceptionEnabled() const;
 
+	bool IsVDZIEnabled() const { return bIsVDZIEnabled; }
+
 	int32 WindowMirrorMode; // how to mirror the display contents to the desktop window: 0 - no mirroring, 1 - single eye, 2 - stereo pair
 	uint32 DebugViewportWidth;
 	uint32 DebugViewportHeight;
