Demote RHI threading warning to message, made more descriptive

diff --git a/Engine/Plugins/Lumin/MagicLeap/Source/MagicLeap/Private/MagicLeapHMD.cpp b/Engine/Plugins/Lumin/MagicLeap/Source/MagicLeap/Private/MagicLeapHMD.cpp
--- a/Engine/Plugins/Lumin/MagicLeap/Source/MagicLeap/Private/MagicLeapHMD.cpp
+++ b/Engine/Plugins/Lumin/MagicLeap/Source/MagicLeap/Private/MagicLeapHMD.cpp
@@ -1364,7 +1364,7 @@ void FMagicLeapHMD::InitDevice_RenderThread()
 		{
 #if PLATFORM_WINDOWS || PLATFORM_LUMIN
 			static const auto* VulkanRHIThread = IConsoleManager::Get().FindTConsoleVariableDataInt(TEXT("r.Vulkan.RHIThread"));
-			UE_LOG(LogMagicLeap, Warning, TEXT("r.Vulkan.RHIThread=%d"), VulkanRHIThread->GetValueOnAnyThread());
+			UE_LOG(LogMagicLeap, Log, TEXT("RHI Thread Usage (r.Vulkan.RHIThread)=%d"), VulkanRHIThread->GetValueOnAnyThread());
 
 			bQueuedGraphicsCreateCall = true;
 
