[NO NON-PII COMMIT MESSAGE]

diff --git a/Engine/Plugins/Lumin/MagicLeap/Source/MagicLeap/Private/MagicLeapARPinComponent.cpp b/Engine/Plugins/Lumin/MagicLeap/Source/MagicLeap/Private/MagicLeapARPinComponent.cpp
--- a/Engine/Plugins/Lumin/MagicLeap/Source/MagicLeap/Private/MagicLeapARPinComponent.cpp
+++ b/Engine/Plugins/Lumin/MagicLeap/Source/MagicLeap/Private/MagicLeapARPinComponent.cpp
@@ -146,8 +146,8 @@ public:
 			{
 				const FAppFramework& AppFramework = static_cast<FMagicLeapHMD*>(GEngine->XRSystem->GetHMDDevice())->GetAppFrameworkConst();
 				const float WorldToMetersScale = AppFramework.GetWorldToMetersScale();
-				FTransform PoseInverse = UHeadMountedDisplayFunctionLibrary::GetTrackingToWorldTransform(GWorld).Inverse();
-				MLVec3f Target = MagicLeap::ToMLVector(PoseInverse.TransformPosition(SearchPoint), WorldToMetersScale);
+				const FTransform WorldToTracking = UHeadMountedDisplayFunctionLibrary::GetTrackingToWorldTransform(GWorld).Inverse();
+				MLVec3f Target = MagicLeap::ToMLVector(WorldToTracking.TransformPosition(SearchPoint), WorldToMetersScale);
 				MLResult Result = MLPersistentCoordinateFrameGetClosest(GetHandle(), &Target, reinterpret_cast<MLCoordinateFrameUID*>(&PinID));
 				UE_CLOG(MLResult_Ok != Result, LogMagicLeap, Error, TEXT("MLPersistentCoordinateFrameGetClosest failed with error %s"), UTF8_TO_TCHAR(MLPersistentCoordinateFrameGetResultString(Result)));
 				ErrorReturn = MLToUnrealPassableWorldError(Result);
@@ -360,8 +360,7 @@ bool UMagicLeapARPinFunctionLibrary::GetARPinPositionAndOrientation(const FGuid&
 	if (AppFramework.GetTransform(*MLCFUID, Pose, FailReason))
 	{
 		FTransform TrackingToWorld = UHeadMountedDisplayFunctionLibrary::GetTrackingToWorldTransform(GWorld);
-		Pose.AddToTranslation(TrackingToWorld.GetLocation());
-		Pose.ConcatenateRotation(TrackingToWorld.Rotator().Quaternion());
+		Pose = Pose * TrackingToWorld;
 		Position = Pose.GetLocation();
 		Orientation = Pose.Rotator();
 		PinFoundInEnvironment = true;
diff --git a/Engine/Plugins/Lumin/MagicLeap/Source/MagicLeap/Private/MeshTrackerComponent.cpp b/Engine/Plugins/Lumin/MagicLeap/Source/MagicLeap/Private/MeshTrackerComponent.cpp
--- a/Engine/Plugins/Lumin/MagicLeap/Source/MagicLeap/Private/MeshTrackerComponent.cpp
+++ b/Engine/Plugins/Lumin/MagicLeap/Source/MagicLeap/Private/MeshTrackerComponent.cpp
@@ -73,8 +73,7 @@ void MLToUnrealBlockInfo(const MLMeshingBlockInfo& MLBlockInfo, const FTransform
 		BlockTransform.SetRotation(rotation);
 	}
 
-	BlockTransform.AddToTranslation(TrackingToWorld.GetLocation());
-	BlockTransform.ConcatenateRotation(TrackingToWorld.Rotator().Quaternion());
+	BlockTransform = BlockTransform * TrackingToWorld;
 	UEBlockInfo.BlockPosition = BlockTransform.GetLocation();
 	UEBlockInfo.BlockOrientation = BlockTransform.Rotator();
 	UEBlockInfo.BlockDimensions = MagicLeap::ToFVectorExtents(MLBlockInfo.extents.extents, WorldToMetersScale);
@@ -505,9 +504,9 @@ void UMeshTrackerComponent::TickComponent(float DeltaTime, enum ELevelTick TickT
 	}
 
 	// Update the bounding box.
-	FTransform PoseInverse = UHeadMountedDisplayFunctionLibrary::GetTrackingToWorldTransform(this).Inverse();
-	Impl->BoundsCenter = PoseInverse.TransformPosition(BoundingVolume->GetComponentLocation());
-	Impl->BoundsRotation = PoseInverse.TransformRotation(BoundingVolume->GetComponentQuat());
+	const FTransform WorldToTracking = UHeadMountedDisplayFunctionLibrary::GetTrackingToWorldTransform(this).Inverse();
+	Impl->BoundsCenter = WorldToTracking.TransformPosition(BoundingVolume->GetComponentLocation());
+	Impl->BoundsRotation = WorldToTracking.TransformRotation(BoundingVolume->GetComponentQuat());
 
 	// Potentially update for changed component parameters
 	if (Impl->Update(*this))
@@ -787,7 +786,7 @@ bool UMeshTrackerComponent::GetMeshResult()
 		}
 		else
 		{
-			FVector VertexOffset = UHeadMountedDisplayFunctionLibrary::GetTrackingToWorldTransform(this).Inverse().GetLocation();
+			const FVector VertexOffset = UHeadMountedDisplayFunctionLibrary::GetTrackingToWorldTransform(this).Inverse().GetLocation();
 			for (uint32_t MeshIndex = 0; MeshIndex < Mesh.data_count; ++ MeshIndex)
 			{
 				const auto &MeshData = Mesh.data[MeshIndex];
diff --git a/Engine/Plugins/Lumin/MagicLeap/Source/MagicLeap/Private/PlanesComponent.cpp b/Engine/Plugins/Lumin/MagicLeap/Source/MagicLeap/Private/PlanesComponent.cpp
--- a/Engine/Plugins/Lumin/MagicLeap/Source/MagicLeap/Private/PlanesComponent.cpp
+++ b/Engine/Plugins/Lumin/MagicLeap/Source/MagicLeap/Private/PlanesComponent.cpp
@@ -205,7 +205,7 @@ void UPlanesComponent::TickComponent(float DeltaTime, enum ELevelTick TickType,
 		return;
 	}
 
-	FTransform PoseTransform = UHeadMountedDisplayFunctionLibrary::GetTrackingToWorldTransform(this);
+	const FTransform TrackingToWorld = UHeadMountedDisplayFunctionLibrary::GetTrackingToWorldTransform(this);
 
 	for (auto& pair : PendingRequests)
 	{
@@ -213,7 +213,7 @@ void UPlanesComponent::TickComponent(float DeltaTime, enum ELevelTick TickType,
 		TArray<MLPlane> resultMLPlanes;
 		resultMLPlanes.AddDefaulted(pair.Value.MaxResults);
 
-		MLResult result = MLPlanesQueryGetResults(Impl->Tracker, pair.Key, resultMLPlanes.GetData(), &outNumResults);
+		MLResult result = MLPlanesQueryGetResultsWithBoundaries(Impl->Tracker, pair.Key, resultMLPlanes.GetData(), &outNumResults, nullptr);
 		switch (result)
 		{
 		case MLResult_Pending:
@@ -254,8 +254,7 @@ void UPlanesComponent::TickComponent(float DeltaTime, enum ELevelTick TickType,
 				}
 
 				planeTransform.ConcatenateRotation(FQuat(FVector(0, 0, 1), PI));
-				planeTransform.AddToTranslation(PoseTransform.GetLocation());
-				planeTransform.ConcatenateRotation(PoseTransform.Rotator().Quaternion());
+				planeTransform = planeTransform * TrackingToWorld;
 				resultPlane.PlanePosition = planeTransform.GetLocation();
 				resultPlane.PlaneOrientation = planeTransform.Rotator();
 				// The plane orientation has the forward axis (X) pointing in the direction of the plane's normal.
@@ -301,7 +300,7 @@ bool UPlanesComponent::RequestPlanes(int32 UserData, const FPlaneResultDelegate&
 	float WorldToMetersScale = AppFramework.GetWorldToMetersScale();
 	check(WorldToMetersScale != 0);
 
-	FTransform PoseInverse = UHeadMountedDisplayFunctionLibrary::GetTrackingToWorldTransform(this).Inverse();
+	const FTransform WorldToTracking = UHeadMountedDisplayFunctionLibrary::GetTrackingToWorldTransform(this).Inverse();
 	FPlane plane;
 
 	MLPlanesQuery query;
@@ -309,8 +308,8 @@ bool UPlanesComponent::RequestPlanes(int32 UserData, const FPlaneResultDelegate&
 	query.flags = UnrealToMLPlanesQueryFlags(QueryFlags);
 	query.min_hole_length = MinHolePerimeter / WorldToMetersScale;
 	query.min_plane_area = MinPlaneArea / (WorldToMetersScale * WorldToMetersScale);
-	query.bounds_center = MagicLeap::ToMLVector(PoseInverse.TransformPosition(SearchVolume->GetComponentLocation()), WorldToMetersScale);
-	query.bounds_rotation = MagicLeap::ToMLQuat(PoseInverse.TransformRotation(SearchVolume->GetComponentQuat()));
+	query.bounds_center = MagicLeap::ToMLVector(WorldToTracking.TransformPosition(SearchVolume->GetComponentLocation()), WorldToMetersScale);
+	query.bounds_rotation = MagicLeap::ToMLQuat(WorldToTracking.TransformRotation(SearchVolume->GetComponentQuat()));
 	query.bounds_extents = MagicLeap::ToMLVectorExtents(SearchVolume->GetScaledBoxExtent(), WorldToMetersScale);
 
 	MLHandle handle;
diff --git a/Engine/Plugins/Lumin/MagicLeap/Source/MagicLeap/Private/RaycastComponent.cpp b/Engine/Plugins/Lumin/MagicLeap/Source/MagicLeap/Private/RaycastComponent.cpp
--- a/Engine/Plugins/Lumin/MagicLeap/Source/MagicLeap/Private/RaycastComponent.cpp
+++ b/Engine/Plugins/Lumin/MagicLeap/Source/MagicLeap/Private/RaycastComponent.cpp
@@ -122,12 +122,12 @@ void URaycastComponent::TickComponent(float DeltaTime, enum ELevelTick TickType,
 
 			// TODO: Should we apply this transform here or expect the user to use the result as a child of the XRPawn like the other features?
 			// This being for raycast, we should probably apply the transform since the result might be used for other than just placing objects.
-			FTransform Pose = UHeadMountedDisplayFunctionLibrary::GetTrackingToWorldTransform(this);
+			const FTransform TrackingToWorld = UHeadMountedDisplayFunctionLibrary::GetTrackingToWorldTransform(this);
 
 			FRaycastHitResult hitResult;
 			hitResult.HitState = MLToUnrealRaycastResultState(result.state);
-			hitResult.HitPoint = Pose.TransformPosition(MagicLeap::ToFVector(result.hitpoint, WorldToMetersScale));
-			hitResult.Normal = Pose.TransformVectorNoScale(MagicLeap::ToFVector(result.normal, 1.0f));
+			hitResult.HitPoint = TrackingToWorld.TransformPosition(MagicLeap::ToFVector(result.hitpoint, WorldToMetersScale));
+			hitResult.Normal = TrackingToWorld.TransformVectorNoScale(MagicLeap::ToFVector(result.normal, 1.0f));
 			hitResult.Confidence = result.confidence;
 			hitResult.UserData = pair.Value.UserData;
 
@@ -169,12 +169,12 @@ bool URaycastComponent::RequestRaycast(const FRaycastQueryParams& RequestParams,
 	const FAppFramework& AppFramework = static_cast<FMagicLeapHMD*>(GEngine->XRSystem->GetHMDDevice())->GetAppFrameworkConst();
 	float WorldToMetersScale = AppFramework.IsInitialized() ? AppFramework.GetWorldToMetersScale() : 100.0f;
 
-	FTransform PoseInverse = UHeadMountedDisplayFunctionLibrary::GetTrackingToWorldTransform(this).Inverse();
+	const FTransform WorldToTracking = UHeadMountedDisplayFunctionLibrary::GetTrackingToWorldTransform(this).Inverse();
 
 	MLRaycastQuery query;
-	query.position = MagicLeap::ToMLVector(PoseInverse.TransformPosition(RequestParams.Position), WorldToMetersScale);
-	query.direction = MagicLeap::ToMLVectorNoScale(PoseInverse.TransformVectorNoScale(RequestParams.Direction));
-	query.up_vector = MagicLeap::ToMLVectorNoScale(PoseInverse.TransformVectorNoScale(RequestParams.UpVector));
+	query.position = MagicLeap::ToMLVector(WorldToTracking.TransformPosition(RequestParams.Position), WorldToMetersScale);
+	query.direction = MagicLeap::ToMLVectorNoScale(WorldToTracking.TransformVectorNoScale(RequestParams.Direction));
+	query.up_vector = MagicLeap::ToMLVectorNoScale(WorldToTracking.TransformVectorNoScale(RequestParams.UpVector));
 	query.width = static_cast<uint32>(RequestParams.Width);
 	query.height = static_cast<uint32>(RequestParams.Height);
 	query.collide_with_unobserved = RequestParams.CollideWithUnobserved;
diff --git a/Engine/Plugins/Lumin/MagicLeap/Source/MagicLeapEyeTracker/Private/MagicLeapVREyeTracker.cpp b/Engine/Plugins/Lumin/MagicLeap/Source/MagicLeapEyeTracker/Private/MagicLeapVREyeTracker.cpp
--- a/Engine/Plugins/Lumin/MagicLeap/Source/MagicLeapEyeTracker/Private/MagicLeapVREyeTracker.cpp
+++ b/Engine/Plugins/Lumin/MagicLeap/Source/MagicLeapEyeTracker/Private/MagicLeapVREyeTracker.cpp
@@ -115,16 +115,14 @@ bool FMagicLeapVREyeTracker::Tick(float DeltaTime)
 			const FAppFramework& AppFramework = static_cast<FMagicLeapHMD*>(GEngine->XRSystem->GetHMDDevice())->GetAppFramework();
 			if (AppFramework.IsInitialized())
 			{
-				FTransform PoseTransform = UHeadMountedDisplayFunctionLibrary::GetTrackingToWorldTransform(GWorld);
+				const FTransform TrackingToWorld = UHeadMountedDisplayFunctionLibrary::GetTrackingToWorldTransform(GWorld);
 
 				//harvest the 3 transforms
 				EFailReason FailReason;
 				FTransform FixationTransform;
 				if (AppFramework.GetTransform(EyeTrackingStaticData.fixation, FixationTransform, FailReason))
 				{
-					FixationTransform.AddToTranslation(PoseTransform.GetLocation());
-					FixationTransform.ConcatenateRotation(PoseTransform.Rotator().Quaternion());
-
+					FixationTransform = FixationTransform * TrackingToWorld;
 					//get focal point
 					UnfilteredEyeTrackingData.WorldAverageGazeConvergencePoint = FixationTransform.GetTranslation();
 				}
@@ -133,8 +131,7 @@ bool FMagicLeapVREyeTracker::Tick(float DeltaTime)
 				FTransform LeftCenterTransform;
 				if (AppFramework.GetTransform(EyeTrackingStaticData.left_center, LeftCenterTransform, FailReason))
 				{
-					LeftCenterTransform.AddToTranslation(PoseTransform.GetLocation());
-					LeftCenterTransform.ConcatenateRotation(PoseTransform.Rotator().Quaternion());
+					LeftCenterTransform = LeftCenterTransform * TrackingToWorld;
 					bLeftTransformValid = true;
 				}
 
@@ -142,8 +139,7 @@ bool FMagicLeapVREyeTracker::Tick(float DeltaTime)
 				FTransform RightCenterTransform;
 				if (AppFramework.GetTransform(EyeTrackingStaticData.right_center, RightCenterTransform, FailReason))
 				{
-					RightCenterTransform.AddToTranslation(PoseTransform.GetLocation());
-					RightCenterTransform.ConcatenateRotation(PoseTransform.Rotator().Quaternion());
+					RightCenterTransform = RightCenterTransform * TrackingToWorld;
 					bRightTransformValid = true;
 				}
 
diff --git a/Engine/Plugins/Lumin/MagicLeapScreens/Source/MagicLeapScreens/Private/MagicLeapScreensPlugin.cpp b/Engine/Plugins/Lumin/MagicLeapScreens/Source/MagicLeapScreens/Private/MagicLeapScreensPlugin.cpp
--- a/Engine/Plugins/Lumin/MagicLeapScreens/Source/MagicLeapScreens/Private/MagicLeapScreensPlugin.cpp
+++ b/Engine/Plugins/Lumin/MagicLeapScreens/Source/MagicLeapScreens/Private/MagicLeapScreensPlugin.cpp
@@ -450,7 +450,7 @@ bool FMagicLeapScreensPlugin::GetScreensTransforms(TArray<FScreenTransform>& Scr
 	if (Result == MLResult_Ok)
 	{
 		// TODO: Param to GetTrackingToWorldTransform is unused (?)
-		FTransform PoseTransform = UHeadMountedDisplayFunctionLibrary::GetTrackingToWorldTransform(nullptr);
+		const FTransform TrackingToWorld = UHeadMountedDisplayFunctionLibrary::GetTrackingToWorldTransform(nullptr);
 		for (uint32 i = 0; i < ScreensInfoList.count; ++i)
 		{
 			MLScreensScreenInfoEx& Entry = ScreensInfoList.entries[i];
@@ -473,8 +473,7 @@ bool FMagicLeapScreensPlugin::GetScreensTransforms(TArray<FScreenTransform>& Scr
 				rotation.Normalize();
 				EntryTransform.SetRotation(rotation);
 			}
-			EntryTransform.AddToTranslation(PoseTransform.GetLocation());
-			EntryTransform.ConcatenateRotation(PoseTransform.Rotator().Quaternion());
+			EntryTransform = EntryTransform * TrackingToWorld;
 
 			ScreenTransform.ScreenPosition = EntryTransform.GetLocation();
 			ScreenTransform.ScreenOrientation = EntryTransform.Rotator();
