Removing unneeded log.

diff --git a/Engine/Plugins/Lumin/MagicLeap/Source/MagicLeap/Private/MagicLeapHMD.cpp b/Engine/Plugins/Lumin/MagicLeap/Source/MagicLeap/Private/MagicLeapHMD.cpp
--- a/Engine/Plugins/Lumin/MagicLeap/Source/MagicLeap/Private/MagicLeapHMD.cpp
+++ b/Engine/Plugins/Lumin/MagicLeap/Source/MagicLeap/Private/MagicLeapHMD.cpp
@@ -1988,7 +1988,6 @@ void FMagicLeapHMD::RestoreBaseProfile()
 void FMagicLeapHMD::EnablePrivileges()
 {
 #if WITH_MLSDK
-	UE_LOG(LogMagicLeap, Warning, TEXT("FMagicLeapHMD::EnablePrivileges"));
 	MLResult Result = MLPrivilegesStartup();
 	bPrivilegesEnabled = (Result == MLResult_Ok);
 	UE_CLOG(!bPrivilegesEnabled, LogMagicLeap, Error, TEXT("MLPrivilegesStartup() "
