Add required AR overrides to Lumin.

diff --git a/Engine/Plugins/Lumin/MagicLeap/Source/MagicLeapAR/Private/LuminARTrackingSystem.cpp b/Engine/Plugins/Lumin/MagicLeap/Source/MagicLeapAR/Private/LuminARTrackingSystem.cpp
--- a/Engine/Plugins/Lumin/MagicLeap/Source/MagicLeapAR/Private/LuminARTrackingSystem.cpp
+++ b/Engine/Plugins/Lumin/MagicLeap/Source/MagicLeapAR/Private/LuminARTrackingSystem.cpp
@@ -74,6 +74,15 @@ bool FLuminARImplementation::OnStartARGameFrame(FWorldContext& WorldContext)
 	return true;
 }
 
+bool FLuminARImplementation::IsARAvailable() const
+{
+	return true;
+}
+
+EARTrackingQualityReason FLuminARImplementation::OnGetTrackingQualityReason() const
+{
+	return EARTrackingQualityReason::None;
+}
 
 EARTrackingQuality FLuminARImplementation::OnGetTrackingQuality() const
 {
diff --git a/Engine/Plugins/Lumin/MagicLeap/Source/MagicLeapAR/Private/LuminARTrackingSystem.h b/Engine/Plugins/Lumin/MagicLeap/Source/MagicLeapAR/Private/LuminARTrackingSystem.h
--- a/Engine/Plugins/Lumin/MagicLeap/Source/MagicLeapAR/Private/LuminARTrackingSystem.h
+++ b/Engine/Plugins/Lumin/MagicLeap/Source/MagicLeapAR/Private/LuminARTrackingSystem.h
@@ -32,7 +32,13 @@ protected:
 	virtual void OnARSystemInitialized() override;
 	virtual bool OnStartARGameFrame(FWorldContext& WorldContext) override;
 
+	/** Returns true/false based on whether AR features are available */
+	virtual bool IsARAvailable() const override;
+
 	virtual EARTrackingQuality OnGetTrackingQuality() const override;
+
+	virtual EARTrackingQualityReason OnGetTrackingQualityReason() const override;
+
 	virtual void OnStartARSession(UARSessionConfig* SessionConfig) override;
 	virtual void OnPauseARSession() override;
 	virtual void OnStopARSession() override;
