MagicLeap VRPreview Crash fix.
-Merging //UE4/Dev-VR to Release-4.23 (//UE4/Release-4.23) cl 8999598
-Do not endframe unless you began frame.
-A fix from MagicLeap.

diff --git a/Engine/Plugins/Lumin/MagicLeap/Source/MagicLeap/Private/MagicLeapCustomPresent.cpp b/Engine/Plugins/Lumin/MagicLeap/Source/MagicLeap/Private/MagicLeapCustomPresent.cpp
--- a/Engine/Plugins/Lumin/MagicLeap/Source/MagicLeap/Private/MagicLeapCustomPresent.cpp
+++ b/Engine/Plugins/Lumin/MagicLeap/Source/MagicLeap/Private/MagicLeapCustomPresent.cpp
@@ -383,6 +383,7 @@ void FMagicLeapCustomPresentOpenGL::FinishRendering()
 			UE_LOG(LogMagicLeap, Error, TEXT("MLGraphicsEndFrame failed with status %d"), Result);
 		}
 	}
+	Plugin->GetCurrentFrameMutable().bBeginFrameSucceeded = false;
 	Plugin->InitializeOldFrameFromRenderFrame();
 #endif //WITH_MLSDK
 }
@@ -610,7 +611,8 @@ void FMagicLeapCustomPresentVulkan::FinishRendering()
 #endif
 		}
 	}
-  
+
+	Plugin->GetCurrentFrameMutable().bBeginFrameSucceeded = false;
 	Plugin->InitializeOldFrameFromRenderFrame();
 #endif // WITH_MLSDK
 }
