Updated ML integration to MLSDK 0.19

diff --git a/Engine/Build/InstalledEngineBuild.xml b/Engine/Build/InstalledEngineBuild.xml
--- a/Engine/Build/InstalledEngineBuild.xml
+++ b/Engine/Build/InstalledEngineBuild.xml
@@ -49,6 +49,8 @@
 	<!-- Whether to force the VS2015 compiler for Windows target platofms. Object files built against the 2015 CRT are compatible with 2017, but not vice-versa. -->
 	<Option Name="ForceVS2015" Restrict="true|false" DefaultValue="false" Description="Whether to force use of the VS2015 compiler for Windows targets. Object files built against the VS2015 CRT are compatible with VS2017 CRT, but not vice versa."/>
 
+	<Option Name="OnPartnerStream" Restrict="true|false" DefaultValue="false" Description="Ignores build features missing from partner streams, like Feature Packs"/>
+
 	<Do If="'$(HostPlatformOnly)' == true">
 		<!-- Activate correct Target Platforms for host -->
 		<Switch>
@@ -96,7 +98,7 @@
 	<Property Name="AndroidArchitectures" Value="armv7+arm64"/>
 
 	<Property Name="AndroidGPUArchitectures" Value="es2"/>
-	<Property Name="LuminGPUArchitectures" Value="es2"/>
+	<Property Name="LuminGPUArchitectures" Value="lumin"/>
 	
 	<!-- Compile flags for Windows Game binaries -->
 	<Property Name="TargetDebugInfoArg" Value="-nodebuginfo"/>
@@ -415,7 +417,6 @@
 
 			<Tag Files="#UE4Game LuminMac" Filter="$(LuminStripFilter)" With="#UE4Game LuminMac Unstripped"/>
 			<Strip BaseDir="$(RootDir)" OutputDir="$(SavedOutput)" Platform="Lumin" Files="#UE4Game LuminMac Unstripped" Tag="#UE4Game LuminMac Stripped"/>
-
 		</Node>
 	</Agent>
 
diff --git a/Engine/Config/BaseEngine.ini b/Engine/Config/BaseEngine.ini
--- a/Engine/Config/BaseEngine.ini
+++ b/Engine/Config/BaseEngine.ini
@@ -1887,9 +1887,6 @@ bAllowIMU=True
 SDKAPILevel=latest
 NDKAPILevel=android-19
 
-[/Script/LuminPlatformEditor.MagicLeapSDKSettings]
-MLSDKPath=(Path="")
-
 [/Script/MagicLeap.MagicLeapSettings]
 bEnableZI=False
 bUseVulkanForZI=True
@@ -1903,11 +1900,13 @@ bProtectedContent=False
 bUseMobileRendering=true
 bBuildWithNvTegraGfxDebugger=False
 Certificate=(FilePath="")
-IconModelPath=(Path="../../Build/Lumin/Resources/Model")
-IconPortalPath=(Path="../../Build/Lumin/Resources/Portal")
+IconModelPath=(Path="Build/Lumin/Resources/Model")
+IconPortalPath=(Path="Build/Lumin/Resources/Portal")
 bUseVulkan=true
 VersionCode=1
-MinimumOSVersion="lumin_1.0"
+MinimumAPILevel=2
+ControllerTrackingType=Tracked
+ControllerTrackingMode=CoordinateFrameUID
 +AppPrivileges=LowLatencyLightwear
 +AppPrivileges=Internet
 +AppPrivileges=LocalAreaNetwork
@@ -1921,6 +1920,8 @@ MinimumOSVersion="lumin_1.0"
 +AppPrivileges=IdentityRead
 +AppPrivileges=AudioRecognizer
 +AppPrivileges=VoiceInput
++AppPrivileges=ControllerPose
+bRemoveDebugInfo=true
 
 [/Script/UnrealEd.CookerSettings]
 DefaultPVRTCQuality=1
diff --git a/Engine/Plugins/Lumin/MagicLeap/MagicLeap.uplugin b/Engine/Plugins/Lumin/MagicLeap/MagicLeap.uplugin
--- a/Engine/Plugins/Lumin/MagicLeap/MagicLeap.uplugin
+++ b/Engine/Plugins/Lumin/MagicLeap/MagicLeap.uplugin
@@ -22,6 +22,7 @@
 		{
 			"Name" : "MagicLeap",
 			"Type" : "Runtime",
+			"LoadingPhase" : "PostConfigInit",
 			"WhitelistPlatforms" :
 			[
 				"Win64",
diff --git a/Engine/Plugins/Lumin/MagicLeap/Resources/custom.css b/Engine/Plugins/Lumin/MagicLeap/Resources/custom.css
--- a/Engine/Plugins/Lumin/MagicLeap/Resources/custom.css
+++ b/Engine/Plugins/Lumin/MagicLeap/Resources/custom.css
@@ -217,7 +217,8 @@ dd {
   background: #F12742;
 }
 
-#nav-tree-contents .children_ul .item > a:first-child img {
+#nav-tree-contents .children_ul .item > a:first-child img,
+#nav-tree-contents .children_ul .item > a:first-child span {
   /* display: none; */
 
   position: relative;
@@ -227,9 +228,15 @@ dd {
   padding: 0 !important;
   filter: brightness(0);
 }
+#nav-tree-contents .children_ul .item > a:first-child span {
+  display: block;
+  font-size: 6px;
+  padding: 12px 0 0 12px !important;
+}
 
 #nav-tree-contents .children_ul .item > a {
   position: relative;
+  text-shadow: none;
 }
 
 #nav-tree-contents .children_ul[style*="display:none"] ~ .item a:before {
@@ -249,8 +256,8 @@ dd {
 
 #nav-tree-contents .children_ul .item > a:first-child {
   position: absolute;
-  top: 8px;
-  left: 12px;
+  top: 0;
+  left: 0;
 }
 
 #nav-tree-contents .children_ul .item .label a,
@@ -404,6 +411,7 @@ hr {
 }
 
 .contents table:not(.memname):not(.mlabels) img[alt="o"],
-.contents table:not(.memname):not(.mlabels) img[alt="\\"] {
-  display: none;
+.contents table:not(.memname):not(.mlabels) img[alt="\\"],
+.contents table .entry > span {
+  display: none !important;
 }
diff --git a/Engine/Plugins/Lumin/MagicLeap/Source/MLSDK/MLSDK.Build.cs b/Engine/Plugins/Lumin/MagicLeap/Source/MLSDK/MLSDK.Build.cs
--- a/Engine/Plugins/Lumin/MagicLeap/Source/MLSDK/MLSDK.Build.cs
+++ b/Engine/Plugins/Lumin/MagicLeap/Source/MLSDK/MLSDK.Build.cs
@@ -36,6 +36,10 @@ public class MLSDK : ModuleRules
 				case UnrealTargetPlatform.Lumin:
 					LibraryPlatformFolder = "lumin";
 					break;
+				default:
+					// This will fail the bIsMLSDKInstalled check, causing WITH_MLSDK to be set to 0 for unsupported platforms.
+					LibraryPlatformFolder = "unsupported";
+					break;
 			}
 			LibraryPath = Path.Combine(LibraryPath, LibraryPlatformFolder);
 
@@ -48,16 +52,10 @@ public class MLSDK : ModuleRules
 				ConfigHierarchy Ini = ConfigCache.ReadHierarchy(ConfigHierarchyType.Engine, ProjectDir, Target.Platform);
 
 				PublicIncludePaths.Add(IncludePath);
-				if (Target.Platform != UnrealTargetPlatform.Lumin)
+				if (Target.Platform == UnrealTargetPlatform.Lumin)
 				{
-					string VirtualDeviceIncludePath = "";
-					if (Ini.TryGetValue("MLSDK", "IncludePath", out VirtualDeviceIncludePath))
-					{
-						PublicIncludePaths.Add(VirtualDeviceIncludePath);
-					}
+					PublicIncludePaths.Add(Target.UEThirdPartySourceDirectory + "Vulkan/Include/vulkan");
 				}
-				//PublicIncludePaths.Add(Path.Combine(MLSDKPath, "lumin/usr/include/vulkan"));
-				PublicIncludePaths.Add(Target.UEThirdPartySourceDirectory + "Vulkan/Include/vulkan");
 
 				string MLSDKLibraryPath = "";
 				Ini.TryGetValue("MLSDK", "LibraryPath", out MLSDKLibraryPath);
@@ -89,11 +87,12 @@ public class MLSDK : ModuleRules
 					"ml_mediaplayer",
 					"ml_musicservice",
 					"ml_perception_client",
+					"ml_platform",
 					"ml_privileges",
+					"ml_purchase",
 					"ml_screens",
 					"ml_secure_storage",
 					"ml_sharedfile",
-					"ml_purchase"
 				};
 
 				if (Target.Platform == UnrealTargetPlatform.Win64)
diff --git a/Engine/Plugins/Lumin/MagicLeap/Source/MLSDK/Public/MagicLeapMath.h b/Engine/Plugins/Lumin/MagicLeap/Source/MLSDK/Public/MagicLeapMath.h
--- a/Engine/Plugins/Lumin/MagicLeap/Source/MLSDK/Public/MagicLeapMath.h
+++ b/Engine/Plugins/Lumin/MagicLeap/Source/MLSDK/Public/MagicLeapMath.h
@@ -23,6 +23,17 @@ namespace MagicLeap
 		return FVector(-InVec3f.xyz.z * WorldToMetersScale, InVec3f.xyz.x * WorldToMetersScale, InVec3f.xyz.y * WorldToMetersScale);
 	}
 
+	FORCEINLINE FVector ToFVectorNoScale(const MLVec3f& InVec3f)
+	{
+		return FVector(-InVec3f.xyz.z, InVec3f.xyz.x, InVec3f.xyz.y);
+	}
+
+	FORCEINLINE FVector ToFVectorExtents(const MLVec3f& InVec3f, float WorldToMetersScale)
+	{
+		// No sign change for FVector::X
+		return FVector(InVec3f.xyz.z * WorldToMetersScale, InVec3f.xyz.x * WorldToMetersScale, InVec3f.xyz.y * WorldToMetersScale);
+	}
+
 	FORCEINLINE MLVec3f ToMLVector(const FVector& InVector, float WorldToMetersScale)
 	{
 		float InverseScale = 1.0f / WorldToMetersScale;
@@ -38,6 +49,17 @@ namespace MagicLeap
 		return ToMLVector(InVector, 1.0f);
 	}
 
+	FORCEINLINE MLVec3f ToMLVectorExtents(const FVector& InVector, float WorldToMetersScale)
+	{
+		float InverseScale = 1.0f / WorldToMetersScale;
+		MLVec3f vec;
+		vec.xyz.x = InVector.Y * InverseScale;
+		vec.xyz.y = InVector.Z * InverseScale;
+		// No sign change for MLVec3f::z
+		vec.xyz.z = InVector.X * InverseScale;
+		return vec;
+	}
+
 	FORCEINLINE FQuat ToFQuat(const MLQuaternionf& InQuat)
 	{
 		return FQuat(-InQuat.z, InQuat.x, InQuat.y, -InQuat.w);
@@ -77,6 +99,75 @@ namespace MagicLeap
 			FPlane(InMat4f.matrix_colmajor[12], InMat4f.matrix_colmajor[13], InMat4f.matrix_colmajor[14], InMat4f.matrix_colmajor[15])
 		);
 	}
+
+	FORCEINLINE FVector ToFVector(const MLMat4f& InMatrix, float WorldToMetersScale)
+	{
+		MLVec3f MLPosition;
+		MLPosition.xyz.x = InMatrix.matrix_colmajor[12];
+		MLPosition.xyz.y = InMatrix.matrix_colmajor[13];
+		MLPosition.xyz.z = InMatrix.matrix_colmajor[14];
+		return ToFVector(MLPosition, WorldToMetersScale);
+	}
+
+	FORCEINLINE FQuat ToFQuat(const MLMat4f& InMatrix)
+	{
+		const float Trace = InMatrix.matrix_colmajor[0] + InMatrix.matrix_colmajor[5] + InMatrix.matrix_colmajor[10];
+
+		MLQuaternionf Output;
+		if (Trace > 0)
+		{
+			const float S = FMath::Sqrt(Trace + 1.0) * 2;
+			Output.w = S * 0.25;
+			Output.x = (InMatrix.matrix_colmajor[6] - InMatrix.matrix_colmajor[9]) / S;
+			Output.y = (InMatrix.matrix_colmajor[8] - InMatrix.matrix_colmajor[2]) / S;
+			Output.z = (InMatrix.matrix_colmajor[1] - InMatrix.matrix_colmajor[4]) / S;
+		}
+		else if (InMatrix.matrix_colmajor[0] > InMatrix.matrix_colmajor[5] && InMatrix.matrix_colmajor[0] > InMatrix.matrix_colmajor[10])
+		{
+			const float S = FMath::Sqrt(InMatrix.matrix_colmajor[0] - InMatrix.matrix_colmajor[5] - InMatrix.matrix_colmajor[10] + 1.0) * 2;
+			Output.x = S * 0.25;
+			Output.w = (InMatrix.matrix_colmajor[6] - InMatrix.matrix_colmajor[9]) / S;
+			Output.y = (InMatrix.matrix_colmajor[1] - InMatrix.matrix_colmajor[4]) / S;
+			Output.z = (InMatrix.matrix_colmajor[8] - InMatrix.matrix_colmajor[2]) / S;
+		}
+		else if (InMatrix.matrix_colmajor[5] > InMatrix.matrix_colmajor[10])
+		{
+			const float S = FMath::Sqrt(InMatrix.matrix_colmajor[5] - InMatrix.matrix_colmajor[0] - InMatrix.matrix_colmajor[10] + 1.0) * 2;
+			Output.y = S * 0.25;
+			Output.w = (InMatrix.matrix_colmajor[8] - InMatrix.matrix_colmajor[2]) / S;
+			Output.x = (InMatrix.matrix_colmajor[1] - InMatrix.matrix_colmajor[4]) / S;
+			Output.z = (InMatrix.matrix_colmajor[6] - InMatrix.matrix_colmajor[9]) / S;
+		}
+		else
+		{
+			const float S = FMath::Sqrt(InMatrix.matrix_colmajor[10] - InMatrix.matrix_colmajor[5] - InMatrix.matrix_colmajor[0] + 1.0) * 2;
+			Output.z = S * 0.25;
+			Output.w = (InMatrix.matrix_colmajor[1] - InMatrix.matrix_colmajor[4]) / S;
+			Output.x = (InMatrix.matrix_colmajor[8] - InMatrix.matrix_colmajor[2]) / S;
+			Output.y = (InMatrix.matrix_colmajor[6] - InMatrix.matrix_colmajor[9]) / S;
+		}
+
+		return ToFQuat(Output);
+	}
+
+	// TODO: Remove this once Screens team fixes setting dimensions when returning the ScreenInfoEx list
+	FORCEINLINE MLVec3f ScaleFromMLMatrix(const MLMat4f& InMatrix)
+	{
+		MLVec3f MLScale;
+		MLScale.x = FMath::Sqrt(
+			FMath::Square(InMatrix.matrix_colmajor[0])
+			+ FMath::Square(InMatrix.matrix_colmajor[1])
+			+ FMath::Square(InMatrix.matrix_colmajor[2]));
+		MLScale.y = FMath::Sqrt(
+			FMath::Square(InMatrix.matrix_colmajor[4])
+			+ FMath::Square(InMatrix.matrix_colmajor[5])
+			+ FMath::Square(InMatrix.matrix_colmajor[6]));
+		MLScale.z = FMath::Sqrt(
+			FMath::Square(InMatrix.matrix_colmajor[8])
+			+ FMath::Square(InMatrix.matrix_colmajor[9])
+			+ FMath::Square(InMatrix.matrix_colmajor[10]));
+		return MLScale;
+	}
 #endif //WITH_MLSDK
 
 	FORCEINLINE FQuat ToUERotator(const FQuat& InMLRotation)
diff --git a/Engine/Plugins/Lumin/MagicLeap/Source/MLSDK/Public/MagicLeapPluginUtil.h b/Engine/Plugins/Lumin/MagicLeap/Source/MLSDK/Public/MagicLeapPluginUtil.h
--- a/Engine/Plugins/Lumin/MagicLeap/Source/MLSDK/Public/MagicLeapPluginUtil.h
+++ b/Engine/Plugins/Lumin/MagicLeap/Source/MLSDK/Public/MagicLeapPluginUtil.h
@@ -5,6 +5,8 @@
 #include "Misc/ConfigCacheIni.h"
 #include "HAL/PlatformProcess.h"
 #include "Misc/Paths.h"
+#include "GenericPlatform/GenericPlatformFile.h"
+#include "Misc/CString.h"
 
 #if PLATFORM_WINDOWS
 #define ML_INCLUDES_START THIRD_PARTY_INCLUDES_START \
@@ -28,51 +30,11 @@ public:
 		// We search various places for the ML API DLLs to support loading alternate
 		// implementations. For example to use VDZI on PC platforms.
 		// Public MLSDK path.
-		FString MLSDK;
 
-		// Give preference to the config setting.
-		FString MLSDKPath;
 		FString MLSDKEnvironmentVariableName = TEXT("MLSDK");
-		FString OriginalMLSDKEnvironmentVariableName = TEXT("MLSDKOriginal");
-
-		//save off the value of MLSDK, if it hasn't been saved before
-		MLSDK = FPlatformMisc::GetEnvironmentVariable(*OriginalMLSDKEnvironmentVariableName);
-		MLSDK.TrimToNullTerminator();
-		if (MLSDK.IsEmpty())
-		{
-			//only once, read the starting value
-			MLSDK = FPlatformMisc::GetEnvironmentVariable(*MLSDKEnvironmentVariableName);
-			FPlatformMisc::SetEnvironmentVar(*OriginalMLSDKEnvironmentVariableName, *MLSDK);
-		}
+		FString MLSDK = FPlatformMisc::GetEnvironmentVariable(*MLSDKEnvironmentVariableName);
 
 
-		GConfig->GetString(TEXT("/Script/LuminPlatformEditor.MagicLeapSDKSettings"), TEXT("MLSDKPath"), MLSDKPath, GEngineIni);
-		MLSDKPath.TrimToNullTerminator();
-		if (MLSDKPath.Len() > 0)
-		{
-			int32 StartIndex = -1;
-			bool bFirstQuote = MLSDKPath.FindChar('"', StartIndex);
-			int32 EndIndex = -1;
-			bool bLastQuote = MLSDKPath.FindLastChar('"', EndIndex);
-			if (StartIndex != 1 && EndIndex != -1 && StartIndex < EndIndex)
-			{
-				++StartIndex;
-				MLSDKPath = MLSDKPath.Mid(StartIndex, EndIndex - StartIndex);
-			}
-			else
-			{
-				MLSDKPath.Empty();
-			}
-
-			if (MLSDKPath.Len() > 0 && FPaths::DirectoryExists(MLSDKPath))
-			{
-				MLSDK = MLSDKPath;
-			}
-		}
-
-		MLSDK.TrimToNullTerminator();
-		FPlatformMisc::SetEnvironmentVar(*MLSDKEnvironmentVariableName, *MLSDK);
-
 		if (CheckForVDZILibraries)
 		{
 			// VDZI search paths: VDZI is only active in PC builds.
@@ -84,13 +46,24 @@ public:
 				DllSearchPaths.Add(VDZILibraryPath);
 			}
 
+			// We also search in the MLSDK VDZI paths for libraries if we have them.
 			if (!MLSDK.IsEmpty())
 			{
-				// The default VDZI dir.
-				DllSearchPaths.Add(FPaths::Combine(*MLSDK, TEXT("VirtualDevice"), TEXT("lib")));
-				// We also need to add the default bin dir as dependent libs are placed there instead
-				// of in the lib directory.
-				DllSearchPaths.Add(FPaths::Combine(*MLSDK, TEXT("VirtualDevice"), TEXT("bin")));
+				TArray<FString> ZIShimPath = GetZIShimPath(MLSDK);
+				if (ZIShimPath.Num() > 0)
+				{
+					DllSearchPaths.Append(GetZIShimPath(MLSDK));
+				}
+				else
+				{
+					// Fallback to adding fixed known paths if we fail to get anything from
+					// the configuration data.
+					// The default VDZI dir.
+					DllSearchPaths.Add(FPaths::Combine(*MLSDK, TEXT("VirtualDevice"), TEXT("lib")));
+					// We also need to add the default bin dir as dependent libs are placed there instead
+					// of in the lib directory.
+					DllSearchPaths.Add(FPaths::Combine(*MLSDK, TEXT("VirtualDevice"), TEXT("bin")));
+				}
 			}
 		}
 #endif
@@ -188,7 +161,125 @@ public:
 		}
 	}
 
+	/** Fills the Variables with the evaluated contents of the SDK Shim discovery data. */
+	bool GetZIShimVariables(const FString & MLSDK, TMap<FString, FString> & ResultVariables) const
+	{
+		// The known path to the paths file.
+		FString SDKShimDiscoveryFile = FPaths::Combine(MLSDK, TEXT(".metadata"), TEXT("sdk_shim_discovery.txt"));
+		if (!FPaths::FileExists(SDKShimDiscoveryFile))
+		{
+			return false;
+		}
+		// Map of variable to value for evaluating the content of the file.
+		// On successful parsing and evaluation we copy the data to the result.
+		TMap<FString, FString> Variables;
+		Variables.Add(TEXT("$(MLSDK)"), MLSDK);
+		// TODO: Determine MLSDK version and set MLSDK_VERSION variable.
+#if PLATFORM_WINDOWS
+		Variables.Add(TEXT("$(HOST)"), TEXT("win64"));
+#elif PLATFORM_LINUX
+		Variables.Add(TEXT("$(HOST)"), TEXT("linux64"));
+#elif PLATFORM_MAC
+		Variables.Add(TEXT("$(HOST)"), TEXT("osx"));
+#endif // PLATFORM_WINDOWS
+		// Single pass algo for evaluating the file:
+		// 1. for each line:
+		// a. for each occurance of $(NAME):
+		// i. replace with Variables[NAME], if no NAME var found ignore replacement.
+		// b. split into var=value, and add to Variables.
+		IPlatformFile& PlatformFile = IPlatformFile::GetPlatformPhysical();
+		TUniquePtr<IFileHandle> File(PlatformFile.OpenRead(*SDKShimDiscoveryFile));
+		if (File)
+		{
+			TArray<uint8> Data;
+			Data.SetNumZeroed(File->Size() + 1);
+			File->Read(Data.GetData(), Data.Num() - 1);
+			FUTF8ToTCHAR TextConvert(reinterpret_cast<ANSICHAR*>(Data.GetData()));
+			const TCHAR * Text = TextConvert.Get();
+			while (Text && *Text)
+			{
+				Text += FCString::Strspn(Text, TEXT("\t "));
+				if (Text[0] == '#' || Text[0] == '\r' || Text[0] == '\n' || !*Text)
+				{
+					// Skip comment or empty lines.
+					Text += FCString::Strcspn(Text, TEXT("\r\n"));
+				}
+				else
+				{
+					// Parse variable=value line.
+					FString Variable(FCString::Strcspn(Text, TEXT("\t= ")), Text);
+					Text += Variable.Len();
+					Text += FCString::Strspn(Text, TEXT("\t= "));
+					FString Value(FCString::Strcspn(Text, TEXT("\r\n")), Text);
+					Text += Value.Len();
+					Value.TrimEndInline();
+					// Eval any var references in both variable and value.
+					int32 EvalCount = 0;
+					do
+					{
+						EvalCount = 0;
+						for (auto VarEntry : Variables)
+						{
+							EvalCount += Variable.ReplaceInline(*VarEntry.Key, *VarEntry.Value);
+							EvalCount += Value.ReplaceInline(*VarEntry.Key, *VarEntry.Value);
+						}
+					} while (EvalCount != 0 && (Variable.Contains(TEXT("$(")) || Value.Contains(TEXT("$("))));
+					// Intern the new variable.
+					Variable = TEXT("$(") + Variable + TEXT(")");
+					Variables.Add(Variable, Value);
+				}
+				// Skip over EOL to get to the next line.
+				if (Text[0] == '\r' && Text[1] == '\n')
+				{
+					Text += 2;
+				}
+				else
+				{
+					Text += 1;
+				}
+			}
+		}
+		// We now need to copy the evaled variables to the result and un-munge them for
+		// plain access. We use them munged for the eval for easier eval replacement above.
+		ResultVariables.Empty(Variables.Num());
+		for (auto VarEntry : Variables)
+		{
+			ResultVariables.Add(VarEntry.Key.Mid(2, VarEntry.Key.Len() - 3), VarEntry.Value);
+		}
+		return true;
+	}
+
 private:
 	TArray<FString> DllSearchPaths;
 	TArray<void*> DllHandles;
+
+	TArray<FString> GetZIShimPath(const FString & MLSDK) const
+	{
+#if PLATFORM_WINDOWS
+		const FString ZIShimPathVar = "ZI_SHIM_PATH_win64";
+#elif PLATFORM_LINUX
+		const FString ZIShimPathVar = "ZI_SHIM_PATH_linux64";
+#elif PLATFORM_MAC
+		const FString ZIShimPathVar = "ZI_SHIM_PATH_osx";
+#else
+		const FString ZIShimPathVar = "";
+#endif // PLATFORM_WINDOWS
+
+		TMap<FString, FString> Variables;
+		if (GetZIShimVariables(MLSDK, Variables) && Variables.Contains(ZIShimPathVar))
+		{
+			// The shim path var we are looking for.
+			FString Value = Variables[ZIShimPathVar];
+			// Since it's a path variable it can have multiple components. Hence
+			// split those out into our result;
+			TArray<FString> Result;
+			for (FString Path; Value.Split(TEXT(";"), &Path, &Value); )
+			{
+				Result.Add(Path);
+			}
+			Result.Add(Value);
+			return Result;
+		}
+		return TArray<FString>();
+	}
 };
diff --git a/Engine/Plugins/Lumin/MagicLeap/Source/MagicLeap/Classes/CameraCaptureComponent.h b/Engine/Plugins/Lumin/MagicLeap/Source/MagicLeap/Classes/CameraCaptureComponent.h
--- a/Engine/Plugins/Lumin/MagicLeap/Source/MagicLeap/Classes/CameraCaptureComponent.h
+++ b/Engine/Plugins/Lumin/MagicLeap/Source/MagicLeap/Classes/CameraCaptureComponent.h
@@ -30,35 +30,80 @@ public:
   virtual void TickComponent(float DeltaTime, enum ELevelTick TickType, FActorComponentTickFunction *ThisTickFunction) override;
 
   /**
-	Instigates a capture image to file task on a separate thread.
+   Delegate used to notify the initiating blueprint when a capture image to file task has completed.
+   @note Although this signals the task as complete, it may have failed or been cancelled.
+   @param bSuccess True if the task succeeded, false otherwise.
+   @param FilePath A string containing the file path to the newly created jpeg.
+ */
+  DECLARE_DYNAMIC_DELEGATE_TwoParams(FCameraCaptureImgToFile, bool, bSuccess, const FString&, FilePath);
+
+  /**
+	Delegate used to pass the captured image back to the initiating blueprint.
+	@note The captured texture will remain in memory for the lifetime of the calling application (if the task succeeds).
+	@param bSuccess True if the task succeeded, false otherwise.
+	@param CaptureTexture A UTexture2D containing the captured image.
+  */
+  DECLARE_DYNAMIC_DELEGATE_TwoParams(FCameraCaptureImgToTexture, bool, bSuccess, UTexture2D*, CaptureTexture);
+
+  /**
+	Delegate used to notify the initiating blueprint of the result of a request to begin recording video.
+	@note Although this signals the task as complete, it may have failed or been cancelled.
+	@param bSuccess True if the task succeeded, false otherwise.
+  */
+  DECLARE_DYNAMIC_DELEGATE_OneParam(FCameraCaptureStartRecording, bool, bSuccess);
+
+  /**
+	Delegate used to notify the initiating blueprint of the result of a request to stop recording video.
+	@note Although this signals the task as complete, it may have failed or been cancelled.
+	@param bSuccess True if the task succeeded, false otherwise.
+	@param FilePath A string containing the path to the newly created mp4.
+  */
+  DECLARE_DYNAMIC_DELEGATE_TwoParams(FCameraCaptureStopRecording, bool, bSuccess, const FString&, FilePath);
+
+  /**
+	Initiates a capture image to file task on a separate thread.
 	@brief The newly created jpeg file will have an automatically generated name which is guaranteed
-		   to be unique.  The user should register event handlers for both the success and fail events.
-		   Upon completion, a successful operation will provide the file path of the newly created jpeg
-		   to the FCameraCaptureImgToFileSuccess event handler.
+		   to be unique.  Upon completion, a successful operation will provide the file path of the newly
+		   created jpeg to the FCameraCaptureImgToFile event handler.
+	@param ResultDelegate The delegate to be notified once the camera image has been saved to a jpeg file.
   */
-  UFUNCTION(BlueprintCallable, Category = "MagicLeap|CameraCapture")
-  bool CaptureImageToFileAsync();
+  UFUNCTION(BlueprintCallable, Category = "CameraCapture|MagicLeap")
+  bool CaptureImageToFileAsync(const FCameraCaptureImgToFile& ResultDelegate);
 
   /** 
-    Instigates a capture image to memory task on a speparate thread.
+    Initiates a capture image to memory task on a speparate thread.
 	@brief The user should register event handlers for both the success and fail events.  Upon completion,
 		   a successful operation will provide a dynamically generated texture containing the captured
 		   image to the FCameraCaptureImgToTextureSuccess event handler.
-    @note The generated texture will be garbage collected when this class is destroyed.
+    @note The generated texture will be garbage collected when this app is destroyed.
+	@param ResultDelegate The delegate to be notified once the camera image has been saved to a texture.
   */
-  UFUNCTION(BlueprintCallable, Category = "MagicLeap|CameraCapture")
-  bool CaptureImageToTextureAsync();
+  UFUNCTION(BlueprintCallable, Category = "CameraCapture|MagicLeap")
+  bool CaptureImageToTextureAsync(const FCameraCaptureImgToTexture& ResultDelegate);
 
   /** 
-    Instigates a capture video to file task on a separate thread.
-	@brief The newly created mp4 file will have an automatically generated name which is guaranteed
-		   to be unique.  The user should register event handlers for both the success and fail events.
-		   Upon completion, a successful operation will provide the file path of the saved mp4 file to the
-		   FCameraCaptureVidToFileSuccess event handler.
-	@param VideoLength The length in seconds of the footage to be captured.
+    Initiates the capturing of video/audio data on a separate thread.
+	@note The system will continue to record video until StopRecordingVideo is called.
+	@param ResultDelegate The delegate to be notified once the recording has begun or failed to begin.
   */
-  UFUNCTION(BlueprintCallable, Category = "MagicLeap|CameraCapture")
-  bool CaptureVideoToFileAsync(float VideoLength);
+  UFUNCTION(BlueprintCallable, Category = "CameraCapture|MagicLeap")
+  bool StartRecordingVideoAsync(const FCameraCaptureStartRecording& ResultDelegate);
+
+  /**
+	Stops the recording and saves the video/audio data to an mp4 file.
+	@note The newly created mp4 file will have an automatically generated name which is guaranteed
+		  to be unique.
+	@param ResultDelegate The delegate to be notified once the video/audio data has been saved to an mp4 file.
+  */
+  UFUNCTION(BlueprintCallable, Category = "CameraCapture|MagicLeap")
+  bool StopRecordingVideoAsync(const FCameraCaptureStopRecording& ResultDelegate);
+
+  /**
+	Gets the capture state of the component.
+	@return True if the component is currently capturing, false otherwise.
+  */
+  UFUNCTION(BlueprintPure, Category = "CameraCapture|MagicLeap")
+  bool IsCapturing() const;
 
   /**
 	Retrieves a handle to the current preview buffer.
@@ -69,84 +114,20 @@ public:
 
 public:
   /**
-    Delegate used to pass log messages from the capture worker thread to the instigating blueprint.
+    Delegate used to pass log messages from the capture worker thread to the initiating blueprint.
 	@note This is useful if the user wishes to have log messages in 3D space.
     @param LogMessage A string containing the log message.
   */
   DECLARE_DYNAMIC_MULTICAST_DELEGATE_OneParam(FCameraCaptureLogMessage, FString, LogMessage);
 
-  /**
-    Delegate used to notify the instigating blueprint of a capture image to file success.
-    @param FilePath A string containing the file path to the newly created jpeg.
-  */
-  DECLARE_DYNAMIC_MULTICAST_DELEGATE_OneParam(FCameraCaptureImgToFileSuccess, FString, FilePath);
-
-  /** Delegate used to notify the instigating blueprint of a capture image to file failure. */
-  DECLARE_DYNAMIC_MULTICAST_DELEGATE(FCameraCaptureImgToFileFailure);
-
-  /**
-    Delegate used to pass the successfully captured image back to the instigating blueprint.
-    @note The captured texture will be cleaned up when this component is destroyed.
-    @param CaptureTexture A UTexture2D containing the captured image.
-  */
-  DECLARE_DYNAMIC_MULTICAST_DELEGATE_OneParam(FCameraCaptureImgToTextureSuccess, UTexture2D*, CaptureTexture);
-
-  /** Delegate used to notify the instigating blueprint of a capture image to texture failure. */
-  DECLARE_DYNAMIC_MULTICAST_DELEGATE(FCameraCaptureImgToTextureFailure);
-
-  /** Delegate used to notify the instigating blueprint of a capture video to file success.
-
-    @param FilePath A string containing the path to the newly created mp4.
-  */
-  DECLARE_DYNAMIC_MULTICAST_DELEGATE_OneParam(FCameraCaptureVidToFileSuccess, FString, FilePath);
-
-  /** Delegate used to notify the instigating blueprint of a capture video to file failure. */
-  DECLARE_DYNAMIC_MULTICAST_DELEGATE(FCameraCaptureVidToFileFailure);
-
-  /** Activated when a log message is generated on the capture thread. */
-  FCameraCaptureLogMessage& OnCaptureLogMessage();
-
-  /** Activated when an image is successfully captured to file on the capture thread. */
-  FCameraCaptureImgToFileSuccess& OnCaptureImgToFileSuccess();
-
-  /** Activated when an image fails to capture to file on the capture thread. */
-  FCameraCaptureImgToFileFailure& OnCaptureImgToFileFailure();
-
-  /** Activated when an image is successfully captured to texture on the capture thread. */
-  FCameraCaptureImgToTextureSuccess& OnCaptureImgToTextureSuccess();
-
-  /** Activated when an image fails to capture to texture on the capture thread. */
-  FCameraCaptureImgToTextureFailure& OnCaptureImgToTextureFailure();
-
-  /** Activated when a video is successfully captured to file on the capture thread. */
-  FCameraCaptureVidToFileSuccess& OnCaptureVidToFileSuccess();
-
-  /** Activated when a video is fails to capture to file on the capture thread. */
-  FCameraCaptureVidToFileFailure& OnCaptureVidToFileFailure();  
-
 private:
-  class FCameraCaptureImpl *Impl;
-
-  UPROPERTY(BlueprintAssignable, Category = "CameraCapture|MagicLeap", meta = (AllowPrivateAccess = true))
-  FCameraCaptureLogMessage CaptureLogMessage;
-
-  UPROPERTY(BlueprintAssignable, Category = "CameraCapture|MagicLeap", meta = (AllowPrivateAccess = true))
-  FCameraCaptureImgToFileSuccess CaptureImgToFileSuccess;
-
-  UPROPERTY(BlueprintAssignable, Category = "CameraCapture|MagicLeap", meta = (AllowPrivateAccess = true))
-  FCameraCaptureImgToFileFailure CaptureImgToFileFailure;
-
-  UPROPERTY(BlueprintAssignable, Category = "CameraCapture|MagicLeap", meta = (AllowPrivateAccess = true))
-  FCameraCaptureImgToTextureSuccess CaptureImgToTextureSuccess;
-
-  UPROPERTY(BlueprintAssignable, Category = "CameraCapture|MagicLeap", meta = (AllowPrivateAccess = true))
-  FCameraCaptureImgToTextureFailure CaptureImgToTextureFailure;
-
-  UPROPERTY(BlueprintAssignable, Category = "CameraCapture|MagicLeap", meta = (AllowPrivateAccess = true))
-  FCameraCaptureVidToFileSuccess CaptureVidToFileSuccess;
-
-  UPROPERTY(BlueprintAssignable, Category = "CameraCapture|MagicLeap", meta = (AllowPrivateAccess = true))
-  FCameraCaptureVidToFileFailure CaptureVidToFileFailure;
+	class FCameraCaptureImpl *Impl;
+	UPROPERTY(BlueprintAssignable, Category = "CameraCapture|MagicLeap", meta = (AllowPrivateAccess = true))
+	FCameraCaptureLogMessage CaptureLogMessage;
+	FCameraCaptureImgToFile CaptureImgToFileResult;
+	FCameraCaptureImgToTexture CaptureImgToTextureResult;
+	FCameraCaptureStartRecording StartRecordingResult;
+	FCameraCaptureStopRecording StopRecordingResult;
 
 private:
   void Log(const FString& LogMessage);
diff --git a/Engine/Plugins/Lumin/MagicLeap/Source/MagicLeap/Classes/ImageTrackerComponent.h b/Engine/Plugins/Lumin/MagicLeap/Source/MagicLeap/Classes/ImageTrackerComponent.h
--- a/Engine/Plugins/Lumin/MagicLeap/Source/MagicLeap/Classes/ImageTrackerComponent.h
+++ b/Engine/Plugins/Lumin/MagicLeap/Source/MagicLeap/Classes/ImageTrackerComponent.h
@@ -37,9 +37,10 @@ public:
 		When the task completes, the instigating blueprint will be notified by either a FSetImageTargetSucceeded
 		or FSetImageTargetFailed event.
 		@param ImageTarget The new texture to be tracked.
+		@return True if the initiation of the target change was successful, false otherwise.
 	*/
 	UFUNCTION(BlueprintCallable, Category = "ImageTracking|MagicLeap")
-	void SetTargetAsync(UTexture2D* ImageTarget);
+	bool SetTargetAsync(UTexture2D* ImageTarget);
 
 	/** Delegate used to notify the instigating blueprint that the target image was successfully set. */
 	DECLARE_DYNAMIC_MULTICAST_DELEGATE(FSetImageTargetSucceeded);
@@ -115,7 +116,6 @@ public:
 	bool bUseUnreliablePose;
 
 private:
-	friend class FImageTrackerImpl;
 	class FImageTrackerImpl *Impl;
 	bool bTick;
 
diff --git a/Engine/Plugins/Lumin/MagicLeap/Source/MagicLeap/Classes/InAppPurchaseComponent.h b/Engine/Plugins/Lumin/MagicLeap/Source/MagicLeap/Classes/InAppPurchaseComponent.h
--- a/Engine/Plugins/Lumin/MagicLeap/Source/MagicLeap/Classes/InAppPurchaseComponent.h
+++ b/Engine/Plugins/Lumin/MagicLeap/Source/MagicLeap/Classes/InAppPurchaseComponent.h
@@ -222,9 +222,7 @@ public:
 	/** Activated when a purchase history request fails. */
 	FGetPurchaseHistoryFailure& OnGetPurchaseHistoryFailure();
 
-private:
-	friend class FInAppPurchaseImpl;
-	class FInAppPurchaseImpl* Impl;
+	void Log(const FString& LogMessage, bool bError = true);
 
 	UPROPERTY(BlueprintAssignable, Category = "InAppPurchase|MagicLeap", meta = (AllowPrivateAccess = true))
 	FInAppPurchaseLogMessage InAppPurchaseLogMessage;
@@ -248,7 +246,7 @@ private:
 	FGetPurchaseHistoryFailure GetPurchaseHistoryFailure;
 
 private:
-	void Log(const FString& LogMessage, bool bError = true);
+	class FInAppPurchaseImpl* Impl;
 };
 
 DECLARE_LOG_CATEGORY_EXTERN(LogInAppPurchase, Verbose, All);
\ No newline at end of file
diff --git a/Engine/Plugins/Lumin/MagicLeap/Source/MagicLeap/Classes/MagicLeapSettings.h b/Engine/Plugins/Lumin/MagicLeap/Source/MagicLeap/Classes/MagicLeapSettings.h
--- a/Engine/Plugins/Lumin/MagicLeap/Source/MagicLeap/Classes/MagicLeapSettings.h
+++ b/Engine/Plugins/Lumin/MagicLeap/Source/MagicLeap/Classes/MagicLeapSettings.h
@@ -17,4 +17,8 @@ public:
 	// Enables 'Zero Iteration mode'. Note: Vulkan rendering will be used by default. Set bUseVulkan to false to use OpenGL instead.
 	UPROPERTY(GlobalConfig, EditAnywhere, Category = "General", Meta = (DisplayName = "Enable Zero Iteration", ConfigRestartRequired = true))
 	bool bEnableZI;
+
+	// Use the editor in Vulkan. If False, OpenGL is used with ZI.
+	UPROPERTY(GlobalConfig, EditAnywhere, Category = "General", Meta = (DisplayName = "Use Vulkan with Zero Iteration (otherwise, OpenGL)", ConfigRestartRequired = true))
+	bool bUseVulkanForZI;
 };
diff --git a/Engine/Plugins/Lumin/MagicLeap/Source/MagicLeap/Classes/MeshBlockSelectorInterface.h b/Engine/Plugins/Lumin/MagicLeap/Source/MagicLeap/Classes/MeshBlockSelectorInterface.h
new file mode 100644
--- /dev/null
+++ b/Engine/Plugins/Lumin/MagicLeap/Source/MagicLeap/Classes/MeshBlockSelectorInterface.h
@@ -0,0 +1,27 @@
+// Copyright 1998-2018 Epic Games, Inc. All Rights Reserved.
+#pragma once
+
+#include "Engine/Engine.h"
+#include "MeshTrackerTypes.h"
+#include "MeshBlockSelectorInterface.generated.h"
+
+UINTERFACE(BlueprintType, Blueprintable)
+class MAGICLEAP_API UMeshBlockSelectorInterface : public UInterface
+{
+	GENERATED_BODY()
+};
+
+/** Interface to select blocks for a mesh request. */
+class MAGICLEAP_API IMeshBlockSelectorInterface
+{
+	GENERATED_BODY()
+
+public:
+	/**
+	 * Given the new mesh information, select the blocks you want to keep and the level of detail for each of those blocks. 
+	 * @param NewMeshInfo Information on the latest mesh blocks available.
+	 * @param RequestedMesh output array containing details of blocks for whom the mesh should be requested.
+	 */
+	UFUNCTION(BlueprintCallable, BlueprintNativeEvent, Category="Meshing|MagicLeap")
+	void SelectMeshBlocks(const FMeshInfo& NewMeshInfo, TArray<FMeshBlockRequest>& RequestedMesh);
+};
diff --git a/Engine/Plugins/Lumin/MagicLeap/Source/MagicLeap/Classes/MeshTrackerComponent.h b/Engine/Plugins/Lumin/MagicLeap/Source/MagicLeap/Classes/MeshTrackerComponent.h
--- a/Engine/Plugins/Lumin/MagicLeap/Source/MagicLeap/Classes/MeshTrackerComponent.h
+++ b/Engine/Plugins/Lumin/MagicLeap/Source/MagicLeap/Classes/MeshTrackerComponent.h
@@ -2,43 +2,11 @@
 
 #pragma once
 
+#include "MeshTrackerTypes.h"
+#include "MeshBlockSelectorInterface.h"
 #include "Components/SceneComponent.h"
 #include "MeshTrackerComponent.generated.h"
 
-/** Type of mesh to query from the underlying system. */
-UENUM(BlueprintType)
-enum class EMeshType : uint8
-{
-	/*! Meshing should be done as triangles. */
-	Triangles,
-	/*! Return mesh vertices as a point cloud. */
-	PointCloud
-};
-
-/** Vertex color mode. */
-UENUM(BlueprintType)
-enum class EMLMeshVertexColorMode : uint8
-{
-	/** Vertex Color is not set. */
-	None    UMETA(DisplayName = "No Vertex Color"),
-	/*! Vertex confidence is interpolated between two specified colors. */
-	Confidence  UMETA(DisplayName = "Vertex Confidence"),
-	/*! Each block is given a color from a list. */
-	Block  UMETA(DisplayName = "Blocks Colored")
-};
-
-/** Discrete level of detail required. */
-UENUM(BlueprintType)
-enum class EMeshLOD : uint8
-{
-	/*! Minimum LOD. */
-	Minimum,
-	/*! Medium LOD. */
-	Medium,
-	/*! Maximum LOD. */
-	Maximum,
-};
-
 /**
 	The MeshTrackerComponent class manages requests for environmental mesh data, processes the results and provides
 	them to the calling system. The calling system is able request environmental mesh data within a specified area.
@@ -47,12 +15,12 @@ enum class EMeshLOD : uint8
 	FOnMeshTrackerUpdated broadcast.
 */
 UCLASS(ClassGroup = MagicLeap, BlueprintType, Blueprintable, EditInlineNew, meta = (BlueprintSpawnableComponent))
-class MAGICLEAP_API UMeshTrackerComponent
-	: public USceneComponent
+class MAGICLEAP_API UMeshTrackerComponent : public USceneComponent, public IMeshBlockSelectorInterface
 {
-	GENERATED_UCLASS_BODY()
+	GENERATED_BODY()
 
 public:
+	UMeshTrackerComponent();
 	/** Destroys the FMeshTrackerImpl instance.*/
 	~UMeshTrackerComponent();
 
@@ -70,15 +38,34 @@ public:
 	UFUNCTION(BlueprintCallable, Category = "Mesh Reconstruction")
 	void DisconnectMRMesh(class UMRMeshComponent* InMRMeshPtr);
 
+	/**
+		Sets the interface to be used for selecting blocks to mesh.
+		@param Selector pointer to the object which implements IMeshBlockSelectorInterface.
+	*/
+	UFUNCTION(BlueprintCallable, Category = "Meshing|MagicLeap")
+	void ConnectBlockSelector(TScriptInterface<IMeshBlockSelectorInterface> Selector);
+
+	/** 
+	 * Disconnects the previously connected IMeshBlockSelectorInterface.
+	 * The default implementation is used this case - all new and updated blocks are meshed with the MeshTrackerComponent'd LevelOfDetail.
+	 */
+	UFUNCTION(BlueprintCallable, Category = "Meshing|MagicLeap")
+	void DisconnectBlockSelector();
+
+	/** IMeshBlockSelectorInterface */
+	UFUNCTION(BlueprintCallable, BlueprintNativeEvent, Category="Meshing|MagicLeap")
+	void SelectMeshBlocks(const FMeshInfo& NewMeshInfo, TArray<FMeshBlockRequest>& RequestedMesh);
+	virtual void SelectMeshBlocks_Implementation(const FMeshInfo& NewMeshInfo, TArray<FMeshBlockRequest>& RequestedMesh) override;
+
 	/**
 		Delegate used by OnMeshUpdated().
-		@param Index The index of the mesh section in the ProceduralMeshComponent that was updated.
+		@param ID The ID of the mesh section in the ProceduralMeshComponent that was updated.
 		@param Vertices List of all vertices in the updated mesh section.
 		@param Triangles List of all triangles in the updated mesh section.
 		@param Normals List of the normals of all triangles in the updated mesh section.
 		@param Confidence List of the confidence values per vertex in the updated mesh section. This can be used to determine if the user needs to scan more.
 	*/
-	DECLARE_DYNAMIC_MULTICAST_DELEGATE_FiveParams(FOnMeshTrackerUpdated, int32, Index, const TArray<FVector>&, Vertices, const TArray<int32>&, Triangles, const TArray<FVector>&, Normals, const TArray<float>&, Confidence);
+	DECLARE_DYNAMIC_MULTICAST_DELEGATE_FiveParams(FOnMeshTrackerUpdated, FGuid, ID, const TArray<FVector>&, Vertices, const TArray<int32>&, Triangles, const TArray<FVector>&, Normals, const TArray<float>&, Confidence);
 
 	/** Activated whenever new information about this mesh tracker is detected. */
 	UPROPERTY(BlueprintAssignable)
@@ -154,12 +141,6 @@ public:
 	UPROPERTY(transient)
 	class UMRMeshComponent* MRMesh;
 
-	/**
-		The procedural mesh generated should bake physics and update its collision data.
-		TODO: Can we calculate this based on the collision flags of the UProceduralMeshComponent?
-		It seems like those are not being respected directly by the component.
-	*/
-
 	/** Polls for and handles the results of the environmental mesh queries. */
 	virtual void TickComponent(float DeltaTime, enum ELevelTick TickType, FActorComponentTickFunction *ThisTickFunction) override;
 
@@ -174,6 +155,11 @@ public:
 #endif
 
 private:
+	void RequestMeshInfo();
+	bool GetMeshInfoResult();
+	void RequestMesh();
+	bool GetMeshResult();
+
 	class FMeshTrackerImpl *Impl;
 
 #if WITH_EDITOR
diff --git a/Engine/Plugins/Lumin/MagicLeap/Source/MagicLeap/Classes/MeshTrackerTypes.h b/Engine/Plugins/Lumin/MagicLeap/Source/MagicLeap/Classes/MeshTrackerTypes.h
new file mode 100644
--- /dev/null
+++ b/Engine/Plugins/Lumin/MagicLeap/Source/MagicLeap/Classes/MeshTrackerTypes.h
@@ -0,0 +1,122 @@
+// Copyright 1998-2018 Epic Games, Inc. All Rights Reserved.
+
+#pragma once
+
+#include "Engine/Engine.h"
+#include "CoreMinimal.h"
+#include "UObject/ObjectMacros.h"
+#include "MeshTrackerTypes.generated.h"
+
+/** Type of mesh to query from the underlying system. */
+UENUM(BlueprintType)
+enum class EMeshType : uint8
+{
+	/** Meshing should be done as triangles. */
+	Triangles,
+	/** Return mesh vertices as a point cloud. */
+	PointCloud
+};
+
+/** Vertex color mode. */
+UENUM(BlueprintType)
+enum class EMLMeshVertexColorMode : uint8
+{
+	/** Vertex Color is not set. */
+	None		UMETA(DisplayName = "No Vertex Color"),
+	/** Vertex confidence is interpolated between two specified colors. */
+	Confidence	UMETA(DisplayName = "Vertex Confidence"),
+	/** Each block is given a color from a list. */
+	Block		UMETA(DisplayName = "Blocks Colored"),
+	/** Each LOD is given a color from a list. */
+	LOD			UMETA(DisplayName = "LODs Colored")
+};
+
+/** Discrete level of detail required. */
+UENUM(BlueprintType)
+enum class EMeshLOD : uint8
+{
+	/** Minimum LOD. */
+	Minimum,
+	/** Medium LOD. */
+	Medium,
+	/** Maximum LOD. */
+	Maximum,
+};
+
+/** State of a block mesh. */
+UENUM(BlueprintType)
+enum class EMeshState : uint8
+{
+	/** Mesh has been created */
+	New,
+	/** Mesh has been updated. */
+	Updated,
+	/** Mesh has been deleted. */
+	Deleted,
+	/** Mesh is unchanged. */
+	Unchanged
+};
+
+/** Representation of a mesh block. */
+USTRUCT(BlueprintType)
+struct FMeshBlockInfo
+{
+	GENERATED_BODY()
+
+public:
+	/** The coordinate frame UID to represent the block. */
+	UPROPERTY(VisibleAnywhere, BlueprintReadOnly, Category = "Meshing|MagicLeap")
+	FGuid BlockID;
+
+	/** The center of the mesh block bounding box. */
+	UPROPERTY(VisibleAnywhere, BlueprintReadOnly, Category = "Meshing|MagicLeap")
+	FVector BlockPosition;
+
+	/** The orientation of the mesh block bounding box.*/
+	UPROPERTY(VisibleAnywhere, BlueprintReadOnly, Category = "Meshing|MagicLeap")
+	FRotator BlockOrientation;
+
+	/** The size of the mesh block bounding box (in Unreal Units). */
+	UPROPERTY(VisibleAnywhere, BlueprintReadOnly, Category = "Meshing|MagicLeap")
+	FVector BlockDimensions;
+
+	/** The timestamp when block was updated. */
+	UPROPERTY(VisibleAnywhere, BlueprintReadOnly, Category = "Meshing|MagicLeap")
+	FTimespan Timestamp;
+
+	/** The state of the mesh block. */
+	UPROPERTY(VisibleAnywhere, BlueprintReadOnly, Category = "Meshing|MagicLeap")
+	EMeshState BlockState;
+};
+
+/** Response structure for the mesh block info. */
+USTRUCT(BlueprintType)
+struct FMeshInfo
+{
+	GENERATED_BODY()
+
+public:
+	/** The response timestamp to a earlier request. */
+	UPROPERTY(VisibleAnywhere, BlueprintReadOnly, Category = "Meshing|MagicLeap")
+	FTimespan Timestamp;
+
+	/** The meshinfo returned by the system. */
+	UPROPERTY(VisibleAnywhere, BlueprintReadOnly, Category = "Meshing|MagicLeap")
+	TArray<FMeshBlockInfo> BlockData;
+};
+
+/** Request structure to get the actual mesh for a block. */
+USTRUCT(BlueprintType)
+struct FMeshBlockRequest
+{
+	GENERATED_BODY()
+
+public:
+	/** The UID to represent the block. */
+	UPROPERTY(EditAnywhere, BlueprintReadWrite, Category = "Meshing|MagicLeap")
+	FGuid BlockID;
+
+	/** The LOD level to request. */
+	UPROPERTY(EditAnywhere, BlueprintReadWrite, Category = "Meshing|MagicLeap")
+	EMeshLOD LevelOfDetail;
+};
diff --git a/Engine/Plugins/Lumin/MagicLeap/Source/MagicLeap/Classes/PlanesComponent.h b/Engine/Plugins/Lumin/MagicLeap/Source/MagicLeap/Classes/PlanesComponent.h
--- a/Engine/Plugins/Lumin/MagicLeap/Source/MagicLeap/Classes/PlanesComponent.h
+++ b/Engine/Plugins/Lumin/MagicLeap/Source/MagicLeap/Classes/PlanesComponent.h
@@ -62,7 +62,11 @@ public:
 
 	/** The flags which describe this plane. TODO: Should be a TSet but that is misbehaving in the editor.*/
 	UPROPERTY(EditAnywhere, BlueprintReadWrite, Category = "Planes|MagicLeap")
-	TArray<EPlaneQueryFlags> PlaneFlags;	
+	TArray<EPlaneQueryFlags> PlaneFlags;
+
+	/** ID of the plane result. This ID is persistent across queries */
+	UPROPERTY(BlueprintReadOnly, Category = "Planes|MagicLeap")
+	FGuid ID;
 };
 
 /** 
@@ -115,10 +119,6 @@ public:
 	UPROPERTY(EditAnywhere, BlueprintReadWrite, Category = "Planes|MagicLeap", meta = (ClampMin = 400))
 	float MinPlaneArea;
 
-	/** Ignore bounds when tracking planes. */
-	UPROPERTY(EditAnywhere, BlueprintReadWrite, Category = "Planes|MagicLeap")
-	bool IgnoreBoundingVolume = false;
-
 	/**
 	  Requests planes with the current value of QueryFlags, SearchVolume and MaxResults.
 	  @param UserData User data for this request. The same data will be included in the result for query identification.
@@ -128,12 +128,18 @@ public:
 	UFUNCTION(BlueprintCallable, Category = "MagicLeap|Planes")
 	bool RequestPlanes(int32 UserData, const FPlaneResultDelegate& ResultDelegate);
 
+	/** Creates the planes tracker handle for the component */
+	virtual void BeginPlay() override;
+
 	/** Polls for and handles the results of the plane queries. */
 	virtual void TickComponent(float DeltaTime, enum ELevelTick TickType, FActorComponentTickFunction *ThisTickFunction) override;
 
 	/** Destroys the interface object to the planes api*/
 	virtual void FinishDestroy() override;
 
+	UPROPERTY()
+	bool IgnoreBoundingVolume_DEPRECATED;
+
 private:
 	struct FPlanesRequestMetaData
 	{
diff --git a/Engine/Plugins/Lumin/MagicLeap/Source/MagicLeap/Classes/PrivilegeUtils.h b/Engine/Plugins/Lumin/MagicLeap/Source/MagicLeap/Classes/PrivilegeUtils.h
new file mode 100644
--- /dev/null
+++ b/Engine/Plugins/Lumin/MagicLeap/Source/MagicLeap/Classes/PrivilegeUtils.h
@@ -0,0 +1,77 @@
+// %BANNER_BEGIN%
+// ---------------------------------------------------------------------
+// %COPYRIGHT_BEGIN%
+//
+// Copyright (c) 2017 Magic Leap, Inc. (COMPANY) All Rights Reserved.
+// Magic Leap, Inc. Confidential and Proprietary
+//
+// NOTICE: All information contained herein is, and remains the property
+// of COMPANY. The intellectual and technical concepts contained herein
+// are proprietary to COMPANY and may be covered by U.S. and Foreign
+// Patents, patents in process, and are protected by trade secret or
+// copyright law. Dissemination of this information or reproduction of
+// this material is strictly forbidden unless prior written permission is
+// obtained from COMPANY. Access to the source code contained herein is
+// hereby forbidden to anyone except current COMPANY employees, managers
+// or contractors who have executed Confidentiality and Non-disclosure
+// agreements explicitly covering such access.
+//
+// The copyright notice above does not evidence any actual or intended
+// publication or disclosure of this source code, which includes
+// information that is confidential and/or proprietary, and is a trade
+// secret, of COMPANY. ANY REPRODUCTION, MODIFICATION, DISTRIBUTION,
+// PUBLIC PERFORMANCE, OR PUBLIC DISPLAY OF OR THROUGH USE OF THIS
+// SOURCE CODE WITHOUT THE EXPRESS WRITTEN CONSENT OF COMPANY IS
+// STRICTLY PROHIBITED, AND IN VIOLATION OF APPLICABLE LAWS AND
+// INTERNATIONAL TREATIES. THE RECEIPT OR POSSESSION OF THIS SOURCE
+// CODE AND/OR RELATED INFORMATION DOES NOT CONVEY OR IMPLY ANY RIGHTS
+// TO REPRODUCE, DISCLOSE OR DISTRIBUTE ITS CONTENTS, OR TO MANUFACTURE,
+// USE, OR SELL ANYTHING THAT IT MAY DESCRIBE, IN WHOLE OR IN PART.
+//
+// %COPYRIGHT_END%
+// --------------------------------------------------------------------
+// %BANNER_END%
+
+#pragma once
+
+#include "UObject/ObjectMacros.h"
+
+/** Priviliges an app can request for from the system. */
+UENUM(BlueprintType)
+enum class EMagicLeapPrivilege : uint8
+{
+	Invalid,
+	AudioRecognizer,
+	BatteryInfo,
+	CameraCapture,
+	WorldReconstruction,
+	InAppPurchase,
+	AudioCaptureMic,
+	DrmCertificates,
+	Occlusion,
+	LowLatencyLightwear,
+	Internet,
+	IdentityRead,
+	BackgroundDownload,
+	BackgroundUpload,
+	MediaDrm,
+	Media,
+	MediaMetadata,
+	PowerInfo,
+	LocalAreaNetwork,
+	VoiceInput,
+	Documents,
+	ConnectBackgroundMusicService,
+	RegisterBackgroundMusicService,
+	PwFoundObjRead,
+	NormalNotificationsUsage,
+	MusicService,
+	ControllerPose,
+	ScreensProvider,
+	GesturesSubscribe,
+	GesturesConfig,
+};
+
+//FString PrivilegeToString(EMagicLeapPrivilege PrivilegeID);
+//EMagicLeapPrivilege StringToPrivilege(const TCHAR* Privilege);
+
diff --git a/Engine/Plugins/Lumin/MagicLeap/Source/MagicLeap/Private/AppEventHandler.cpp b/Engine/Plugins/Lumin/MagicLeap/Source/MagicLeap/Private/AppEventHandler.cpp
--- a/Engine/Plugins/Lumin/MagicLeap/Source/MagicLeap/Private/AppEventHandler.cpp
+++ b/Engine/Plugins/Lumin/MagicLeap/Source/MagicLeap/Private/AppEventHandler.cpp
@@ -49,8 +49,12 @@ namespace MagicLeap
 		return UTF8_TO_TCHAR("");
 	}
 
-	IAppEventHandler::IAppEventHandler(TArray<MLPrivilegeID> InRequiredPrivilegeIDs)
-	: bAllPrivilegesInSync(false)
+	IAppEventHandler::IAppEventHandler(const TArray<MLPrivilegeID>& InRequiredPrivilegeIDs)
+	: OnAppShutDownHandler(nullptr)
+	, OnAppTickHandler(nullptr)
+	, OnAppPauseHandler(nullptr)
+	, OnAppResumeHandler(nullptr)
+	, bAllPrivilegesInSync(false)
 	, bWasSystemEnabledOnPause(false)
 	{
 		RequiredPrivileges.Reserve(InRequiredPrivilegeIDs.Num());
@@ -116,6 +120,14 @@ namespace MagicLeap
 	}
 #endif // WITH_MLSDK
 
+	void IAppEventHandler::OnAppShutDown()
+	{
+		if (OnAppShutDownHandler)
+		{
+			OnAppShutDownHandler();
+		}
+	}
+
 	void IAppEventHandler::OnAppTick()
 	{
 #if WITH_MLSDK
@@ -172,6 +184,19 @@ namespace MagicLeap
 			}
 		}
 #endif // WITH_MLSDK
+
+		if (OnAppTickHandler)
+		{
+			OnAppTickHandler();
+		}
+	}
+
+	void IAppEventHandler::OnAppPause()
+	{
+		if (OnAppPauseHandler)
+		{
+			OnAppPauseHandler();
+		}
 	}
 
 	void IAppEventHandler::OnAppResume()
@@ -184,10 +209,15 @@ namespace MagicLeap
 			ItRequiredPrivilege.Value.State = EPrivilegeState::NotYetRequested;
 		}
 #endif // WITH_MLSDK
+
+		if (OnAppResumeHandler)
+		{
+			OnAppResumeHandler();
+		}
 	}
 
-	void IAppEventHandler::AsyncDestroy()
+	bool IAppEventHandler::AsyncDestroy()
 	{
-		FAppFramework::AsyncDestroy(this);
+		return FAppFramework::AsyncDestroy(this);
 	}
 }
\ No newline at end of file
diff --git a/Engine/Plugins/Lumin/MagicLeap/Source/MagicLeap/Private/AppEventHandler.h b/Engine/Plugins/Lumin/MagicLeap/Source/MagicLeap/Private/AppEventHandler.h
--- a/Engine/Plugins/Lumin/MagicLeap/Source/MagicLeap/Private/AppEventHandler.h
+++ b/Engine/Plugins/Lumin/MagicLeap/Source/MagicLeap/Private/AppEventHandler.h
@@ -39,13 +39,15 @@ namespace MagicLeap
 	class MAGICLEAP_API IAppEventHandler
 	{
 	public:
+		typedef TFunction<void()> FEventHandler;
+
 #if WITH_MLSDK
 		/** 
 			Adds the IAppEventHandler instance to the application's list of IAppEventHandler instances.
 			Populates a RequiredPrivileges list based on the privilge ids passed via InRequiredPrivileges;
 			@param InRequiredPrivileges The list of privilge ids required by the calling system.
 		*/
-		IAppEventHandler(TArray<MLPrivilegeID> InRequiredPrivileges);
+		IAppEventHandler(const TArray<MLPrivilegeID>& InRequiredPrivileges);
 #endif // WITH_MLSDK
 
 		IAppEventHandler();
@@ -57,7 +59,7 @@ namespace MagicLeap
 			Can be overridden by inheriting class that needs to destroy certain api interfaces before the perception stack is
 			closed down.
 		*/
-		virtual void OnAppShutDown() {}
+		virtual void OnAppShutDown();
 
 		/**
 			Use to check status of privilege requests.
@@ -67,7 +69,7 @@ namespace MagicLeap
 		/**
 			Can be overridden by inheriting class in order to pause its system.
 		*/
-		virtual void OnAppPause() {}
+		virtual void OnAppPause();
 
 		/**
 			Can be overridden by inheriting class in order to resume its system.
@@ -87,12 +89,52 @@ namespace MagicLeap
 			Pushes this object onto a worker thread so that it's blocking destructor can be called without locking up the update thread.
 			@note This should only be called by objects that have a blocking destructor and are no longer referenced.
 		*/
-		void AsyncDestroy();
+		bool AsyncDestroy();
+
+		/**
+			Use this as an alternative to overriding the OnAppShutDown function.  This allows you to use IAppEventHandler
+			as and aggregate class rather than an ancestor. 
+		*/
+		void SetOnAppShutDownHandler(FEventHandler&& InOnAppShutDownHandler)
+		{
+			OnAppShutDownHandler = MoveTemp(InOnAppShutDownHandler);
+		}
+
+		/**
+			Use this as an alternative to overriding the OnAppTick function.  This allows you to use IAppEventHandler
+			as and aggregate class rather than an ancestor.
+		*/
+		void SetOnAppTickHandler(FEventHandler&& InOnAppTickHandler)
+		{
+			OnAppTickHandler = MoveTemp(InOnAppTickHandler);
+		}
+
+		/**
+			Use this as an alternative to overriding the OnAppPause function.  This allows you to use IAppEventHandler
+			as and aggregate class rather than an ancestor.
+		*/
+		void SetOnAppPauseHandler(FEventHandler&& InOnAppPauseHandler)
+		{
+			OnAppPauseHandler = MoveTemp(InOnAppPauseHandler);
+		}
+
+		/**
+			Use this as an alternative to overriding the OnAppResume function.  This allows you to use IAppEventHandler
+			as and aggregate class rather than an ancestor.
+		*/
+		void SetOnAppResumeHandler(FEventHandler&& InOnAppResumeHandler)
+		{
+			OnAppResumeHandler = MoveTemp(InOnAppResumeHandler);
+		}
 
 	protected:
 #if WITH_MLSDK
 		TMap<MLPrivilegeID, FRequiredPrivilege> RequiredPrivileges;
 #endif // WITH_MLSDK
+		FEventHandler OnAppShutDownHandler;
+		FEventHandler OnAppTickHandler;
+		FEventHandler OnAppPauseHandler;
+		FEventHandler OnAppResumeHandler;
 		bool bAllPrivilegesInSync;
 		bool bWasSystemEnabledOnPause;
 		FCriticalSection CriticalSection;
diff --git a/Engine/Plugins/Lumin/MagicLeap/Source/MagicLeap/Private/AppFramework.cpp b/Engine/Plugins/Lumin/MagicLeap/Source/MagicLeap/Private/AppFramework.cpp
--- a/Engine/Plugins/Lumin/MagicLeap/Source/MagicLeap/Private/AppFramework.cpp
+++ b/Engine/Plugins/Lumin/MagicLeap/Source/MagicLeap/Private/AppFramework.cpp
@@ -9,6 +9,7 @@
 #include "Misc/CoreDelegates.h"
 #include "RenderingThread.h"
 #include "MagicLeapPluginUtil.h" // for ML_INCLUDES_START/END
+#include "IMagicLeapModule.h"
 
 #if WITH_MLSDK
 ML_INCLUDES_START
@@ -19,6 +20,8 @@ ML_INCLUDES_END
 TArray<MagicLeap::IAppEventHandler*> FAppFramework::EventHandlers;
 FCriticalSection FAppFramework::EventHandlersCriticalSection;
 MagicLeap::FAsyncDestroyer* FAppFramework::AsyncDestroyer = nullptr;
+TMap<FName, IMagicLeapModule*> FAppFramework::RegisteredModules;
+
 
 FAppFramework::FAppFramework()
 {}
@@ -90,7 +93,6 @@ void FAppFramework::ApplicationPauseDelegate()
 {
 	UE_LOG(LogMagicLeap, Log, TEXT("+++++++ ML AppFramework APP PAUSE ++++++"));
 
-
 	if (GEngine)
 	{
 		saved_max_fps_ = GEngine->GetMaxFPS();
@@ -122,6 +124,7 @@ void FAppFramework::ApplicationResumeDelegate()
 {
 	UE_LOG(LogMagicLeap, Log, TEXT("+++++++ ML AppFramework APP RESUME ++++++"));
 
+	// Resume rendering
 	// Resume rendering
 	FMagicLeapHMD * const HMD = GEngine ? static_cast<FMagicLeapHMD*>(GEngine->XRSystem->GetHMDDevice()) : nullptr;
 	if (HMD)
@@ -343,7 +346,30 @@ void FAppFramework::RemoveEventHandler(MagicLeap::IAppEventHandler* EventHandler
 	EventHandlers.Remove(EventHandler);
 }
 
-void FAppFramework::AsyncDestroy(MagicLeap::IAppEventHandler* InEventHandler)
+bool FAppFramework::AsyncDestroy(MagicLeap::IAppEventHandler* InEventHandler)
+{
+	if (AsyncDestroyer != nullptr)
+	{
+		AsyncDestroyer->AddRaw(InEventHandler);
+		return true;
+	}
+
+	return false;
+}
+
+void FAppFramework::RegisterMagicLeapModule(IMagicLeapModule* InModule)
+{
+	checkf(RegisteredModules.Find(InModule->GetName()) == nullptr, TEXT("MagicLeapModule %s has already been registered!"), *InModule->GetName().ToString());
+	RegisteredModules.Add(InModule->GetName(), InModule);
+}
+
+void FAppFramework::UnregisterMagicLeapModule(IMagicLeapModule* InModule)
+{
+	RegisteredModules.Remove(InModule->GetName());
+}
+
+IMagicLeapModule* FAppFramework::GetMagicLeapModule(FName InName)
 {
-	AsyncDestroyer->AddRaw(InEventHandler);
+	IMagicLeapModule** MagicLeapModule = RegisteredModules.Find(InName);
+	return MagicLeapModule ? *MagicLeapModule : nullptr;
 }
\ No newline at end of file
diff --git a/Engine/Plugins/Lumin/MagicLeap/Source/MagicLeap/Private/AppFramework.h b/Engine/Plugins/Lumin/MagicLeap/Source/MagicLeap/Private/AppFramework.h
--- a/Engine/Plugins/Lumin/MagicLeap/Source/MagicLeap/Private/AppFramework.h
+++ b/Engine/Plugins/Lumin/MagicLeap/Source/MagicLeap/Private/AppFramework.h
@@ -21,6 +21,7 @@ ML_INCLUDES_END
 #include "ImageTrackerRunnable.h"
 
 class FMagicLeapHMD;
+class IMagicLeapModule;
 struct FTrackingFrame;
 
 enum class EFailReason : uint8
@@ -83,7 +84,11 @@ public:
 
 	static void AddEventHandler(MagicLeap::IAppEventHandler* InEventHandler);
 	static void RemoveEventHandler(MagicLeap::IAppEventHandler* InEventHandler);
-	static void AsyncDestroy(MagicLeap::IAppEventHandler* InEventHandler);
+	static bool AsyncDestroy(MagicLeap::IAppEventHandler* InEventHandler);
+
+	static void RegisterMagicLeapModule(IMagicLeapModule* InModule);
+	static void UnregisterMagicLeapModule(IMagicLeapModule* InModule);
+	static IMagicLeapModule* GetMagicLeapModule(FName InName);
 
 private:
 	const FTrackingFrame* GetCurrentFrame() const;
@@ -106,6 +111,7 @@ private:
 	static MagicLeap::FAsyncDestroyer* AsyncDestroyer;
 	TSharedPtr<FCameraCaptureRunnable, ESPMode::ThreadSafe> CameraCaptureRunnable;
 	TSharedPtr<FImageTrackerRunnable, ESPMode::ThreadSafe> ImageTrackerRunnable;
+	static TMap<FName, IMagicLeapModule*> RegisteredModules;
 };
 
 DEFINE_LOG_CATEGORY_STATIC(LogMagicLeap, Log, All);
\ No newline at end of file
diff --git a/Engine/Plugins/Lumin/MagicLeap/Source/MagicLeap/Private/CameraCaptureComponent.cpp b/Engine/Plugins/Lumin/MagicLeap/Source/MagicLeap/Private/CameraCaptureComponent.cpp
--- a/Engine/Plugins/Lumin/MagicLeap/Source/MagicLeap/Private/CameraCaptureComponent.cpp
+++ b/Engine/Plugins/Lumin/MagicLeap/Source/MagicLeap/Private/CameraCaptureComponent.cpp
@@ -29,7 +29,7 @@ class FCameraCaptureImpl : public MagicLeap::IAppEventHandler
 {
 public:
 	FCameraCaptureImpl()
-	: bCapturing(false)
+	: CurrentCaptureTaskType(CaptureTaskType::None)
 	{
 #if PLATFORM_LUMIN
 		checkf(GEngine, TEXT("[FCameraCaptureImpl::FCameraCaptureImpl()] GEngine is null!"));
@@ -58,12 +58,12 @@ public:
 	bool TryCaptureImageToFile()
 	{
 #if PLATFORM_LUMIN
-		if (!bCapturing)
+		if (CurrentCaptureTaskType == CaptureTaskType::None)
 		{
-			bCapturing = true;
+			CurrentCaptureTaskType = CaptureTaskType::ImageToFile;
 			FCaptureMessage Msg;
 			Msg.Type = CaptureMsgType::Request;
-			Msg.CaptureType = MLCameraCaptureType_Image;
+			Msg.CaptureType = CurrentCaptureTaskType;
 			Msg.Requester = this;
 			CameraCaptureRunnable->ProcessCaptureMessage(Msg);
 			return true;
@@ -75,12 +75,12 @@ public:
 	bool TryCaptureImageToTexture()
 	{
 #if PLATFORM_LUMIN
-		if (!bCapturing)
+		if (CurrentCaptureTaskType == CaptureTaskType::None)
 		{
-			bCapturing = true;
+			CurrentCaptureTaskType = CaptureTaskType::ImageToTexture;
 			FCaptureMessage Msg;
 			Msg.Type = CaptureMsgType::Request;
-			Msg.CaptureType = MLCameraCaptureType_ImageRaw;
+			Msg.CaptureType = CurrentCaptureTaskType;
 			Msg.Requester = this;
 			CameraCaptureRunnable->ProcessCaptureMessage(Msg);
 			return true;
@@ -89,16 +89,32 @@ public:
 		return false;
 	}
 
-	bool TryCaptureVideoToFile(float InDuration)
+	bool TryCaptureVideoToFile()
 	{
 #if PLATFORM_LUMIN
-		if (!bCapturing)
+		if (CurrentCaptureTaskType == CaptureTaskType::None)
 		{
-			bCapturing = true;
+			CurrentCaptureTaskType = CaptureTaskType::StartVideoToFile;
 			FCaptureMessage Msg;
 			Msg.Type = CaptureMsgType::Request;
-			Msg.CaptureType = MLCameraCaptureType_Video;
-			Msg.Duration = InDuration;
+			Msg.CaptureType = CurrentCaptureTaskType;
+			Msg.Requester = this;
+			CameraCaptureRunnable->ProcessCaptureMessage(Msg);
+			return true;
+		}
+#endif //PLATFORM_LUMIN
+		return false;
+	}
+
+	bool TryStopCaptureVideoToFile()
+	{
+#if PLATFORM_LUMIN
+		if (CurrentCaptureTaskType == CaptureTaskType::StartVideoToFile)
+		{
+			CurrentCaptureTaskType = CaptureTaskType::StopVideoToFile;
+			FCaptureMessage Msg;
+			Msg.Type = CaptureMsgType::Request;
+			Msg.CaptureType = CurrentCaptureTaskType;
 			Msg.Requester = this;
 			CameraCaptureRunnable->ProcessCaptureMessage(Msg);
 			return true;
@@ -110,7 +126,7 @@ public:
 	bool TryGetResult(FCaptureMessage& OutMsg)
 	{
 #if PLATFORM_LUMIN
-		if (bCapturing && CameraCaptureRunnable->OutgoingMessages.Peek(OutMsg))
+		if (CurrentCaptureTaskType != CaptureTaskType::None && CameraCaptureRunnable->OutgoingMessages.Peek(OutMsg))
 		{
 			if (OutMsg.Requester == this)
 			{
@@ -123,7 +139,7 @@ public:
 	}
 
 public:
-	bool bCapturing;
+	CaptureTaskType CurrentCaptureTaskType;
 	TSharedPtr<FCameraCaptureRunnable, ESPMode::ThreadSafe> CameraCaptureRunnable;
 };
 
@@ -140,7 +156,10 @@ UCameraCaptureComponent::~UCameraCaptureComponent()
 #if PLATFORM_LUMIN
 	if (Impl)
 	{
-		Impl->AsyncDestroy();
+		if (!Impl->AsyncDestroy())
+	    {
+			delete Impl;
+    	}
 		Impl = nullptr;
 	}
 #else
@@ -175,89 +194,48 @@ void UCameraCaptureComponent::TickComponent(float DeltaTime, enum ELevelTick Tic
     {
       switch (Msg.CaptureType)
       {
-      case MLCameraCaptureType_Image:
-      {
-        if (Msg.Success)
-        {
-          CaptureImgToFileSuccess.Broadcast(Msg.FilePath);
-        }
-        else
-        {
-          CaptureImgToFileFailure.Broadcast();
-        }
-      }
-      break;
-
-      case MLCameraCaptureType_ImageRaw:
-      {
-        if (Msg.Success)
-        {
-          CaptureImgToTextureSuccess.Broadcast(Msg.Texture);
-        }
-        else
-        {
-          CaptureImgToTextureFailure.Broadcast();
-        }
-      }
-      break;
-
-      case MLCameraCaptureType_Video:
-      {
-        if (Msg.Success)
-        {
-          CaptureVidToFileSuccess.Broadcast(Msg.FilePath);
-        }
-        else
-        {
-          CaptureVidToFileFailure.Broadcast();
-        }
-      }
+	  case CaptureTaskType::ImageToFile:
+	  {
+		  CaptureImgToFileResult.ExecuteIfBound(Msg.Success, Msg.FilePath);
+		  Impl->CurrentCaptureTaskType = CaptureTaskType::None;
+	  }
+	  break;
+
+	  case CaptureTaskType::ImageToTexture:
+	  {
+		  CaptureImgToTextureResult.ExecuteIfBound(Msg.Success, Msg.Texture);
+		  Impl->CurrentCaptureTaskType = CaptureTaskType::None;
+	  }
+	  break;
+
+	  case CaptureTaskType::StartVideoToFile:
+	  {
+		  StartRecordingResult.ExecuteIfBound(Msg.Success);
+		  // do not reset Impl->CurrentCaptureTaskType if start recording is successful
+		  // as this constitutes an ongoing capture state
+		  if (!Msg.Success)
+		  {
+			  Impl->CurrentCaptureTaskType = CaptureTaskType::None;
+		  }
+	  }
+	  break;
+
+	  case CaptureTaskType::StopVideoToFile:
+	  {
+		  StopRecordingResult.ExecuteIfBound(Msg.Success, Msg.FilePath);
+		  Impl->CurrentCaptureTaskType = CaptureTaskType::None;
+	  }
       break;
       }
-
-      Impl->bCapturing = false;
     }
   }
 #endif //WITH_MLSDK
 }
 
-UCameraCaptureComponent::FCameraCaptureLogMessage& UCameraCaptureComponent::OnCaptureLogMessage()
-{
-  return CaptureLogMessage;
-}
-
-UCameraCaptureComponent::FCameraCaptureImgToFileSuccess& UCameraCaptureComponent::OnCaptureImgToFileSuccess()
-{
-  return CaptureImgToFileSuccess;
-}
-
-UCameraCaptureComponent::FCameraCaptureImgToFileFailure& UCameraCaptureComponent::OnCaptureImgToFileFailure()
-{
-  return CaptureImgToFileFailure;
-}
-
-UCameraCaptureComponent::FCameraCaptureImgToTextureSuccess& UCameraCaptureComponent::OnCaptureImgToTextureSuccess()
-{
-  return CaptureImgToTextureSuccess;
-}
-
-UCameraCaptureComponent::FCameraCaptureImgToTextureFailure& UCameraCaptureComponent::OnCaptureImgToTextureFailure()
-{
-  return CaptureImgToTextureFailure;
-}
-
-UCameraCaptureComponent::FCameraCaptureVidToFileSuccess& UCameraCaptureComponent::OnCaptureVidToFileSuccess()
-{
-  return CaptureVidToFileSuccess;
-}
-
-UCameraCaptureComponent::FCameraCaptureVidToFileFailure& UCameraCaptureComponent::OnCaptureVidToFileFailure()
+bool UCameraCaptureComponent::CaptureImageToFileAsync(const FCameraCaptureImgToFile& ResultDelegate)
 {
-  return CaptureVidToFileFailure;
-}
+  CaptureImgToFileResult = ResultDelegate;
 
-bool UCameraCaptureComponent::CaptureImageToFileAsync()
-{
   if (Impl->TryCaptureImageToFile())
   {
     return true;
@@ -267,8 +245,10 @@ bool UCameraCaptureComponent::CaptureImageToFileAsync()
   return false;
 }
 
-bool UCameraCaptureComponent::CaptureImageToTextureAsync()
+bool UCameraCaptureComponent::CaptureImageToTextureAsync(const FCameraCaptureImgToTexture& ResultDelegate)
 {
+  CaptureImgToTextureResult = ResultDelegate;
+
   if (Impl->TryCaptureImageToTexture())
   {
     return true;
@@ -278,9 +258,11 @@ bool UCameraCaptureComponent::CaptureImageToTextureAsync()
   return false;
 }
 
-bool UCameraCaptureComponent::CaptureVideoToFileAsync(float InDuration)
+bool UCameraCaptureComponent::StartRecordingVideoAsync(const FCameraCaptureStartRecording& ResultDelegate)
 {
-  if (Impl->TryCaptureVideoToFile(InDuration))
+  StartRecordingResult = ResultDelegate;
+
+  if (Impl->TryCaptureVideoToFile())
   {
     return true;
   }
@@ -289,6 +271,24 @@ bool UCameraCaptureComponent::CaptureVideoToFileAsync(float InDuration)
   return false;
 }
 
+bool UCameraCaptureComponent::StopRecordingVideoAsync(const FCameraCaptureStopRecording& ResultDelegate)
+{
+	StopRecordingResult = ResultDelegate;
+
+	if (Impl->TryStopCaptureVideoToFile())
+	{
+		return true;
+	}
+
+	Log("Camera capture already in progress!");
+	return false;
+}
+
+bool UCameraCaptureComponent::IsCapturing() const
+{
+	return Impl && Impl->CurrentCaptureTaskType != CaptureTaskType::None;
+}
+
 int64 UCameraCaptureComponent::GetPreviewHandle()
 {
 #if PLATFORM_LUMIN
@@ -301,5 +301,5 @@ int64 UCameraCaptureComponent::GetPreviewHandle()
 void UCameraCaptureComponent::Log(const FString& LogMessage)
 {
   UE_LOG(LogCameraCapture, Log, TEXT("%s"), *LogMessage);
-  CaptureLogMessage.Broadcast(LogMessage);
+  CaptureLogMessage.Broadcast(FString::Printf(TEXT("<br>%s"), *LogMessage));
 }
\ No newline at end of file
diff --git a/Engine/Plugins/Lumin/MagicLeap/Source/MagicLeap/Private/CameraCaptureRunnable.cpp b/Engine/Plugins/Lumin/MagicLeap/Source/MagicLeap/Private/CameraCaptureRunnable.cpp
--- a/Engine/Plugins/Lumin/MagicLeap/Source/MagicLeap/Private/CameraCaptureRunnable.cpp
+++ b/Engine/Plugins/Lumin/MagicLeap/Source/MagicLeap/Private/CameraCaptureRunnable.cpp
@@ -42,6 +42,13 @@
 #if PLATFORM_LUMIN
 #include "Lumin/LuminAffinity.h"
 #endif // PLATFORM_LUMIN
+#include "MagicLeapPluginUtil.h" // for ML_INCLUDES_START/END
+
+#if WITH_MLSDK
+ML_INCLUDES_START
+#include <ml_media_error.h>
+ML_INCLUDES_END
+#endif //WITH_MLSDK
 
 DEFINE_LOG_CATEGORY(LogCameraCaptureRunnable);
 
@@ -59,10 +66,6 @@ FCameraCaptureRunnable::FCameraCaptureRunnable()
 	, StopTaskCounter(0)
 	, Semaphore(nullptr)
 	, bCameraConnected(false)
-	, RetryConnectWaitTime(0.5f)
-#if WITH_MLSDK
-	, CameraOutput(nullptr)
-#endif // WITH_MLSDK
 	, ImgExtension(".jpeg")
 	, VidExtension(".mp4")
 	, bPaused(false)
@@ -89,51 +92,21 @@ uint32 FCameraCaptureRunnable::Run()
 	{
 		if (bPaused)
 		{
-			if (bCameraConnected)
-			{
-				MLResult Result = MLCameraDisconnect();
-				UE_CLOG(Result != MLResult_Ok, LogCameraCaptureRunnable, Error, TEXT("MLCameraDisconnect failed with error %s!"), UTF8_TO_TCHAR(MLGetResultString(Result)));
-				bCameraConnected = false;
-			}
-
-			Semaphore->Wait();
+			Pause();
 		}
-
-		if (IncomingMessages.Dequeue(CurrentMessage))
+		else if (!IncomingMessages.IsEmpty())
 		{
-			while (!bCameraConnected)
+			if (!bCameraConnected)
 			{
-				if (GetPrivilegeStatus(MLPrivilegeID_CameraCapture) == MagicLeap::EPrivilegeState::Granted)
-				{
-					MLResult Result = MLCameraConnect();
-					if (Result == MLResult_Ok)
-					{
-						bCameraConnected = true;
-						MLCameraDeviceStatusCallbacks deviceStatusCallbacks;
-						FMemory::Memset(&deviceStatusCallbacks, 0, sizeof(MLCameraDeviceStatusCallbacks));
-						deviceStatusCallbacks.on_preview_buffer_available = OnPreviewBufferAvailable;
-						Result = MLCameraSetDeviceStatusCallbacks(&deviceStatusCallbacks, nullptr);
-						if (Result != MLResult_Ok)
-						{
-							Log(FString::Printf(TEXT("MLCameraSetDeviceStatusCallbacks failed with error %s!"), UTF8_TO_TCHAR(MLGetResultString(Result))));
-						}
-					}
-					else
-					{
-						bCameraConnected = false;
-						Log(FString::Printf(TEXT("MLCameraConnect failed with error %s!  Retrying in %.2f seconds"), UTF8_TO_TCHAR(MLGetResultString(Result)), RetryConnectWaitTime));
-						IncomingMessages.Empty();
-						FPlatformProcess::Sleep(RetryConnectWaitTime);
-					}
-				}
-				else
-				{
-					Log(TEXT("Cannot connect to camera due to lack of privilege!"));
-					break;
-				}
+				bCameraConnected = TryConnect();
 			}
-
-			BeginCapture();
+			else
+			{
+				DoCaptureTask();
+			}
+		}
+		else
+		{
 			Semaphore->Wait();
 		}
 	}
@@ -141,7 +114,7 @@ uint32 FCameraCaptureRunnable::Run()
 	if (bCameraConnected)
 	{
 		MLResult Result = MLCameraDisconnect();
-		UE_CLOG(Result != MLResult_Ok, LogCameraCaptureRunnable, Error, TEXT("MLCameraDisconnect failed with error %s!"), UTF8_TO_TCHAR(MLGetResultString(Result)));
+		UE_CLOG(Result != MLResult_Ok, LogCameraCaptureRunnable, Error, TEXT("MLCameraDisconnect failed with error %s!"), UTF8_TO_TCHAR(MLMediaResultGetString(Result)));
 	}
 #endif //WITH_MLSDK
 	return 0;
@@ -163,258 +136,315 @@ void FCameraCaptureRunnable::Stop()
 	Thread = nullptr;
 }
 
-void FCameraCaptureRunnable::BeginCapture()
+void FCameraCaptureRunnable::OnAppPause()
+{
+	IAppEventHandler::OnAppPause();
+	bPaused = true;
+	Semaphore->Trigger();
+}
+
+void FCameraCaptureRunnable::OnAppResume()
+{
+	IAppEventHandler::OnAppResume();
+	bPaused = false;
+	Semaphore->Trigger();
+}
+
+void FCameraCaptureRunnable::OnAppShutDown()
+{
+	Stop();
+}
+
+void FCameraCaptureRunnable::ProcessCaptureMessage(const FCaptureMessage& InMsg)
 {
+	IncomingMessages.Enqueue(InMsg);
+	// wake up the worker to process the event
+	Semaphore->Trigger();
+}
+
 #if WITH_MLSDK
-	MLHandle Handle = ML_INVALID_HANDLE;
-	MLResult Result = MLCameraPrepareCapture(CurrentMessage.CaptureType, &Handle);
-	if (Result != MLResult_Ok)
-	{
-		Log(FString::Printf(TEXT("MLCameraPrepareCapture failed with error %s!  Camera capture aborted!"), UTF8_TO_TCHAR(MLGetResultString(Result))));
-		EndCapture(false);
-		return;
-	}
+void FCameraCaptureRunnable::OnPreviewBufferAvailable(MLHandle Output, void *Data)
+{
+	(void)Data;
+	FCameraCaptureRunnable::PreviewHandle.Set(static_cast<int64>(Output));
+}
+#endif //WITH_MLSDK
 
-	switch (CurrentMessage.CaptureType)
+bool FCameraCaptureRunnable::TryConnect()
+{
+#if WITH_MLSDK
+	IncomingMessages.Peek(CurrentTask); // this is purely so that any log messages go to the correct requester
+	if (GetPrivilegeStatus(MLPrivilegeID_CameraCapture) != MagicLeap::EPrivilegeState::Granted)
 	{
-	case MLCameraCaptureType_Image:
+		Log(TEXT("Cannot connect to camera due to lack of privilege!"));
+		return false;
+	}
+	else
 	{
-		Log(TEXT("Beginning capture image to file."));
+		if (bPaused) return false;
 
-#if PLATFORM_LUMIN
-		IPlatformFile& PlatformFile = IPlatformFile::GetPlatformPhysical();
-		// This module is only for Lumin so this is fine for now.
-		FLuminPlatformFile* LuminPlatformFile = static_cast<FLuminPlatformFile*>(&PlatformFile);
-		UniqueFileName = LuminPlatformFile->ConvertToLuminPath(FPaths::CreateTempFilename(*FPaths::ProjectSavedDir(), TEXT("Img_"), *ImgExtension), true);
-#endif
-		Result = MLCameraCaptureImage(TCHAR_TO_UTF8(*UniqueFileName));
-		if (Result != MLResult_Ok)
-		{
-			Log(FString::Printf(TEXT("MLCameraCaptureImage failed with error %s!  Camera capture aborted!"), UTF8_TO_TCHAR(MLGetResultString(Result))));
-			EndCapture(false);
-			return;
-		}
+		MLResult Result = MLCameraConnect();
 
-		EndCapture(true);
-	}
-	break;
+		if (bPaused) return false;
 
-	case MLCameraCaptureType_ImageRaw:
-	{
-		CameraOutput = nullptr;
-
-		Log(TEXT("Beginning capture image to texture."));
-		Result = MLCameraCaptureImageRaw();
 		if (Result != MLResult_Ok)
 		{
-			Log(FString::Printf(TEXT("MLCameraCaptureImageRaw failed with error %s!  Camera capture aborted!"), UTF8_TO_TCHAR(MLGetResultString(Result))));
-			EndCapture(false);
-			return;
+			Log(FString::Printf(TEXT("MLCameraConnect failed with error %s!"), UTF8_TO_TCHAR(MLMediaResultGetString(Result))));
+			CancelIncomingTasks();
+			return false;
 		}
-
-		Result = MLCameraGetImageStream(&CameraOutput);
-		if (Result != MLResult_Ok)
+		else
 		{
-			Log(FString::Printf(TEXT("MLCameraGetImageStream failed with error %s!  Camera capture aborted!"), UTF8_TO_TCHAR(MLGetResultString(Result))));
-			EndCapture(false);
-			return;
+			FMemory::Memset(&DeviceStatusCallbacks, 0, sizeof(MLCameraDeviceStatusCallbacks));
+			DeviceStatusCallbacks.on_preview_buffer_available = OnPreviewBufferAvailable;
+			Result = MLCameraSetDeviceStatusCallbacks(&DeviceStatusCallbacks, nullptr);
+			if (Result != MLResult_Ok)
+			{
+				Log(FString::Printf(TEXT("MLCameraSetDeviceStatusCallbacks failed with error %s!"), UTF8_TO_TCHAR(MLMediaResultGetString(Result))));
+			}
 		}
+	}
+#endif
 
-		if (CameraOutput->plane_count == 0)
-		{
-			Log(TEXT("Invalid plane_count!  Camera capture aborted!"));
-			EndCapture(false);
-			return;
-		}
+	return true;
+}
 
-		EndCapture(true);
-	}
-	break;
+bool FCameraCaptureRunnable::DoCaptureTask()
+{
+	bool bSuccess = false;
+	IncomingMessages.Dequeue(CurrentTask);
 
-	case MLCameraCaptureType_Video:
+#if WITH_MLSDK
+	switch (CurrentTask.CaptureType)
 	{
-		Log(TEXT("Beginning capture video to file."));
-		if (GetPrivilegeStatus(MLPrivilegeID_AudioRecognizer) == MagicLeap::EPrivilegeState::Granted &&
-			GetPrivilegeStatus(MLPrivilegeID_AudioCaptureMic) == MagicLeap::EPrivilegeState::Granted &&
-			GetPrivilegeStatus(MLPrivilegeID_VoiceInput) == MagicLeap::EPrivilegeState::Granted)
-		{
-#if PLATFORM_LUMIN
-			IPlatformFile& PlatformFile = IPlatformFile::GetPlatformPhysical();
-			// This module is only for Lumin so this is fine for now.
-			FLuminPlatformFile* LuminPlatformFile = static_cast<FLuminPlatformFile*>(&PlatformFile);
-			UniqueFileName = LuminPlatformFile->ConvertToLuminPath(FPaths::CreateTempFilename(*FPaths::ProjectSavedDir(), TEXT("Vid_"), *VidExtension), true);
-#endif
-			Result = MLCameraCaptureVideoStart(TCHAR_TO_UTF8(*UniqueFileName));
-			if (Result != MLResult_Ok)
-			{
-				Log(FString::Printf(TEXT("MLCameraCaptureVideoStart failed with error %s!  Video capture aborted!"), UTF8_TO_TCHAR(MLGetResultString(Result))));
-				EndCapture(false);
-				return;
-			}
+	case CaptureTaskType::ImageToFile: bSuccess = CaptureImageToFile(); break;
+	case CaptureTaskType::ImageToTexture: bSuccess = CaptureImageToTexture(); break;
+	case CaptureTaskType::StartVideoToFile: bSuccess = StartRecordingVideo(); break;
+	case CaptureTaskType::StopVideoToFile: bSuccess = StopRecordingVideo(); break;
+	}
+#endif //WITH_MLSDK
 
-			FPlatformProcess::Sleep(CurrentMessage.Duration);
-			EndCapture(true);
-		}
-		else
-		{
-			Log(TEXT("Cannot capture video due to lack of privilege!"));
-			EndCapture(false);
-		}
+	if (bPaused) return false;
+
+	FCaptureMessage Response = CurrentTask;
+	Response.Type = CaptureMsgType::Response;
+	Response.Success = bSuccess;
+	OutgoingMessages.Enqueue(Response);
+
+	return bSuccess;
+}
 
+bool FCameraCaptureRunnable::CaptureImageToFile()
+{
+#if WITH_MLSDK
+	Log(TEXT("Beginning capture image to file."));
+	MLHandle Handle = ML_INVALID_HANDLE;
+	MLResult Result = MLCameraPrepareCapture(MLCameraCaptureType_Image, &Handle);
+	if (Result != MLResult_Ok)
+	{
+		Log(FString::Printf(TEXT("MLCameraPrepareCapture failed with error %s!  Camera capture aborted!"), UTF8_TO_TCHAR(MLMediaResultGetString(Result))));
+		return false;
 	}
-	break;
+
+	if (bPaused) return false;
+
+#if PLATFORM_LUMIN
+	IPlatformFile& PlatformFile = IPlatformFile::GetPlatformPhysical();
+	// This module is only for Lumin so this is fine for now.
+	FLuminPlatformFile* LuminPlatformFile = static_cast<FLuminPlatformFile*>(&PlatformFile);
+	UniqueFileName = LuminPlatformFile->ConvertToLuminPath(FPaths::CreateTempFilename(*FPaths::ProjectSavedDir(), TEXT("Img_"), *ImgExtension), true);
+#endif
+	Result = MLCameraCaptureImage(TCHAR_TO_UTF8(*UniqueFileName));
+	if (Result != MLResult_Ok)
+	{
+		Log(FString::Printf(TEXT("MLCameraCaptureImage failed with error %s!  Camera capture aborted!"), UTF8_TO_TCHAR(MLMediaResultGetString(Result))));
+		return false;
 	}
-#endif //WITH_MLSDK
+
+	Log(FString::Printf(TEXT("Captured image to %s"), *UniqueFileName));
+	CurrentTask.FilePath = UniqueFileName;
+#endif
+	return true;
 }
 
-void FCameraCaptureRunnable::EndCapture(bool InSuccess)
+bool FCameraCaptureRunnable::CaptureImageToTexture()
 {
 #if WITH_MLSDK
-	CurrentMessage.Type = CaptureMsgType::Response;
-	CurrentMessage.Success = InSuccess;
-	CurrentMessage.FilePath = UniqueFileName;
+	Log(TEXT("Beginning capture image to texture."));
+	MLCameraOutput* CameraOutput = nullptr;
+	MLHandle Handle = ML_INVALID_HANDLE;
+	MLResult Result = MLCameraPrepareCapture(MLCameraCaptureType_ImageRaw, &Handle);
+	if (Result != MLResult_Ok)
+	{
+		Log(FString::Printf(TEXT("MLCameraPrepareCapture failed with error %s!  Camera capture aborted!"), UTF8_TO_TCHAR(MLMediaResultGetString(Result))));
+		return false;
+	}
+
+	if (bPaused) return false;
 
-	if (!InSuccess)
+	Result = MLCameraCaptureImageRaw();
+	if (Result != MLResult_Ok)
 	{
-		uint32 DeviceStatus = 0;
-		MLResult Result = MLCameraGetDeviceStatus(&DeviceStatus);
-		if (Result == MLResult_Ok)
-		{
-			if ((MLCameraDeviceStatusFlag_Available & DeviceStatus) != 0)
-			{
-				Log(TEXT("Device status = MLCameraDeviceStatusFlag_Available"));
-			}
-			else if ((MLCameraDeviceStatusFlag_Opened & DeviceStatus) != 0)
-			{
-				Log(TEXT("Device status = MLCameraDeviceStatusFlag_Opened"));
-			}
-			else if ((MLCameraDeviceStatusFlag_Disconnected  & DeviceStatus) != 0)
-			{
-				Log(TEXT("Device status = MLCameraDeviceStatusFlag_Disconnected "));
-			}
-			else if ((MLCameraDeviceStatusFlag_Error & DeviceStatus) != 0)
-			{
-				Log(TEXT("Device status = MLCameraDeviceStatusFlag_Error"));
-				MLCameraError CameraError;
-				if (MLCameraGetErrorCode(&CameraError) == MLResult_Ok)
-				{
-					switch (CameraError)
-					{
-					case MLCameraError_None: Log(TEXT("Error = MLCameraError_None")); break;
-					case MLCameraError_Invalid: Log(TEXT("Error = MLCameraError_Invalid")); break;
-					case MLCameraError_Disabled: Log(TEXT("Error = MLCameraError_Disabled")); break;
-					case MLCameraError_DeviceFailed: Log(TEXT("Error = MLCameraError_DeviceFailed")); break;
-					case MLCameraError_ServiceFailed: Log(TEXT("Error = MLCameraError_ServiceFailed")); break;
-					case MLCameraError_CaptureFailed: Log(TEXT("Error = MLCameraError_CaptureFailed")); break;
-					case MLCameraError_Unknown: Log(TEXT("Error = MLCameraError_Unknown")); break;
-					}
-				}
-				else
-				{
-					Log(TEXT("MLCameraGetErrorCode failed!"));
-				}
-			}
-		}
-		else
-		{
-			Log(FString::Printf(TEXT("MLCameraGetDeviceStatus failed with error %s!"), UTF8_TO_TCHAR(MLGetResultString(Result))));
-		}
+		Log(FString::Printf(TEXT("MLCameraCaptureImageRaw failed with error %s!  Camera capture aborted!"), UTF8_TO_TCHAR(MLMediaResultGetString(Result))));
+		return false;
 	}
 
-	if (InSuccess)
+	if (bPaused) return false;
+
+	Result = MLCameraGetImageStream(&CameraOutput);
+	if (Result != MLResult_Ok)
 	{
-		switch (CurrentMessage.CaptureType)
-		{
-		case MLCameraCaptureType_Image:
-		{
-			Log(FString::Printf(TEXT("Captured image to %s"), *UniqueFileName));
-		}
-		break;
+		Log(FString::Printf(TEXT("MLCameraGetImageStream failed with error %s!  Camera capture aborted!"), UTF8_TO_TCHAR(MLMediaResultGetString(Result))));
+		return false;
+	}
 
-		case MLCameraCaptureType_ImageRaw:
-		{
-			MLCameraPlaneInfo& ImageInfo = CameraOutput->planes[0];
-			if (ImageWrapper.IsValid() && ImageWrapper->SetCompressed(ImageInfo.data, ImageInfo.size))
-			{
-				const TArray<uint8>* RawData = NULL;
-				if (ImageWrapper->GetRaw(ImageWrapper->GetFormat(), 8, RawData))
-				{
-					Log(FString::Printf(TEXT("ImageWrapper width=%d height=%d size=%d"), ImageWrapper->GetWidth(), ImageWrapper->GetHeight(), RawData->Num()));
-					UTexture2D* CaptureTexture = UTexture2D::CreateTransient(ImageWrapper->GetWidth(), ImageWrapper->GetHeight(), EPixelFormat::PF_B8G8R8A8);
-					CaptureTexture->AddToRoot();
-					FTexture2DMipMap& Mip = CaptureTexture->PlatformData->Mips[0];
-					void* Data = Mip.BulkData.Lock(LOCK_READ_WRITE);
-					FMemory::Memcpy(Data, RawData->GetData(), Mip.BulkData.GetBulkDataSize());
-					Mip.BulkData.Unlock();
-					CaptureTexture->UpdateResource();
-					CurrentMessage.Texture = CaptureTexture;
-				}
-			}
-		}
-		break;
+	if (bPaused) return false;
 
-		case MLCameraCaptureType_Video:
+	if (CameraOutput->plane_count == 0)
+	{
+		Log(TEXT("Invalid plane_count!  Camera capture aborted!"));
+		return false;
+	}
+
+	MLCameraPlaneInfo& ImageInfo = CameraOutput->planes[0];
+	if (ImageWrapper.IsValid() && ImageWrapper->SetCompressed(ImageInfo.data, ImageInfo.size))
+	{
+		const TArray<uint8>* RawData = NULL;
+		if (ImageWrapper->GetRaw(ImageWrapper->GetFormat(), 8, RawData))
 		{
-			MLResult Result = MLCameraCaptureVideoStop();
-			if (Result != MLResult_Ok)
-			{
-				Log(FString::Printf(TEXT("MLCameraCaptureVideoStop failed with error %s!"), UTF8_TO_TCHAR(MLGetResultString(Result))));
-			}
-			else
-			{
-				Log(FString::Printf(TEXT("Captured video to %s"), *UniqueFileName));
-			}
-		}
-		break;
+			Log(FString::Printf(TEXT("ImageWrapper width=%d height=%d size=%d"), ImageWrapper->GetWidth(), ImageWrapper->GetHeight(), RawData->Num()));
+			UTexture2D* CaptureTexture = UTexture2D::CreateTransient(ImageWrapper->GetWidth(), ImageWrapper->GetHeight(), EPixelFormat::PF_R8G8B8A8);
+			CaptureTexture->AddToRoot();
+			FTexture2DMipMap& Mip = CaptureTexture->PlatformData->Mips[0];
+			void* Data = Mip.BulkData.Lock(LOCK_READ_WRITE);
+			FMemory::Memcpy(Data, RawData->GetData(), Mip.BulkData.GetBulkDataSize());
+			Mip.BulkData.Unlock();
+			CaptureTexture->UpdateResource();
+			CurrentTask.Texture = CaptureTexture;
 		}
 	}
 
-	// clear any left over events from this capture session
-	IncomingMessages.Empty();
-	// signal the update thread to handle the capture result
-	OutgoingMessages.Enqueue(CurrentMessage);
-#endif //WITH_MLSDK
+#endif
+	return true;
 }
 
-void FCameraCaptureRunnable::Log(const FString& Info)
+bool FCameraCaptureRunnable::StartRecordingVideo()
 {
-	FCaptureMessage Msg;
-	Msg.Type = CaptureMsgType::Log;
-	Msg.Log = Info;
-	Msg.Requester = CurrentMessage.Requester;
-	OutgoingMessages.Enqueue(Msg);
-}
+#if WITH_MLSDK
+	Log(TEXT("Beginning capture video to file."));
+	MLHandle Handle = ML_INVALID_HANDLE;
+	MLResult Result = MLCameraPrepareCapture(MLCameraCaptureType_Video, &Handle);
+	if (Result != MLResult_Ok)
+	{
+		Log(FString::Printf(TEXT("MLCameraPrepareCapture failed with error %s!  Camera capture aborted!"), UTF8_TO_TCHAR(MLMediaResultGetString(Result))));
+		return false;
+	}
 
-void FCameraCaptureRunnable::ProcessCaptureMessage(const FCaptureMessage& InMsg)
-{
-	IncomingMessages.Enqueue(InMsg);
-	// wake up the worker to process the event
-	Semaphore->Trigger();
+	if (bPaused) return false;
+
+	if (GetPrivilegeStatus(MLPrivilegeID_AudioRecognizer) != MagicLeap::EPrivilegeState::Granted)
+	{
+		Log(TEXT("Cannot capture video due to lack of privilege!"));
+		return false;
+	}
+
+	if (bPaused) return false;
+
+	if (GetPrivilegeStatus(MLPrivilegeID_AudioCaptureMic) != MagicLeap::EPrivilegeState::Granted)
+	{
+		Log(TEXT("Cannot capture video due to lack of privilege!"));
+		return false;
+	}
+
+	if (bPaused) return false;
+
+	if (GetPrivilegeStatus(MLPrivilegeID_VoiceInput) != MagicLeap::EPrivilegeState::Granted)
+	{
+		Log(TEXT("Cannot capture video due to lack of privilege!"));
+		return false;
+	}
+
+	if (bPaused) return false;
+
+#if PLATFORM_LUMIN
+	IPlatformFile& PlatformFile = IPlatformFile::GetPlatformPhysical();
+	// This module is only for Lumin so this is fine for now.
+	FLuminPlatformFile* LuminPlatformFile = static_cast<FLuminPlatformFile*>(&PlatformFile);
+	UniqueFileName = LuminPlatformFile->ConvertToLuminPath(FPaths::CreateTempFilename(*FPaths::ProjectSavedDir(), TEXT("Vid_"), *VidExtension), true);
+#endif
+	Result = MLCameraCaptureVideoStart(TCHAR_TO_UTF8(*UniqueFileName));
+	if (Result != MLResult_Ok)
+	{
+		Log(FString::Printf(TEXT("MLCameraCaptureVideoStart failed with error %s!  Video capture aborted!"), UTF8_TO_TCHAR(MLMediaResultGetString(Result))));
+		return false;
+	}
+#endif
+
+	return true;
 }
 
-void FCameraCaptureRunnable::OnAppPause()
+bool FCameraCaptureRunnable::StopRecordingVideo()
 {
-	IAppEventHandler::OnAppPause();
-	bPaused = true;
-	Semaphore->Trigger();
+#if WITH_MLSDK
+	MLResult Result = MLCameraCaptureVideoStop();
+	if (Result != MLResult_Ok)
+	{
+		Log(FString::Printf(TEXT("MLCameraCaptureVideoStop failed with error %s!"), UTF8_TO_TCHAR(MLMediaResultGetString(Result))));
+		return false;
+	}
+	else
+	{
+		Log(FString::Printf(TEXT("Captured video to %s"), *UniqueFileName));
+	}
+
+	CurrentTask.FilePath = UniqueFileName;
+#endif
+	return true;
 }
 
-void FCameraCaptureRunnable::OnAppResume()
+void FCameraCaptureRunnable::Log(const FString& Info)
 {
-	IAppEventHandler::OnAppResume();
-	bPaused = false;
-	Semaphore->Trigger();
+	FCaptureMessage Msg;
+	Msg.Type = CaptureMsgType::Log;
+	Msg.Log = Info;
+	Msg.Requester = CurrentTask.Requester;
+	OutgoingMessages.Enqueue(Msg);
+	UE_LOG(LogCameraCaptureRunnable, Log, TEXT("%s"), *Info);
 }
 
-void FCameraCaptureRunnable::OnAppShutDown()
+void FCameraCaptureRunnable::Pause()
 {
-	Stop();
+#if WITH_MLSDK
+	// Cancel the current video recording (if one is active).
+	if (CurrentTask.CaptureType == CaptureTaskType::StartVideoToFile)
+	{
+		StopRecordingVideo();
+		FCaptureMessage Response = CurrentTask;
+		Response.Type = CaptureMsgType::Response;
+		Response.Success = false;
+		OutgoingMessages.Enqueue(Response);
+	}
+	// Cancel any incoming tasks.
+	CancelIncomingTasks();
+	// Disconnect camera if connected
+	if (bCameraConnected)
+	{
+		MLResult Result = MLCameraDisconnect();
+		UE_CLOG(Result != MLResult_Ok, LogCameraCaptureRunnable, Error, TEXT("MLCameraDisconnect failed with error %s!"), UTF8_TO_TCHAR(MLMediaResultGetString(Result)));
+		bCameraConnected = false;
+	}
+	// Wait for signal from resume call.
+	Semaphore->Wait();
+#endif
 }
 
-#if WITH_MLSDK
-void FCameraCaptureRunnable::OnPreviewBufferAvailable(MLHandle Output, void *Data)
+void FCameraCaptureRunnable::CancelIncomingTasks()
 {
-	(void)Data;
-	FCameraCaptureRunnable::PreviewHandle.Set(static_cast<int64>(Output));
-}
-#endif //WITH_MLSDK
\ No newline at end of file
+	FCaptureMessage CancelledTask;
+	while (IncomingMessages.Dequeue(CancelledTask))
+	{
+		CancelledTask.Type = CaptureMsgType::Response;
+		CancelledTask.Success = false;
+		OutgoingMessages.Enqueue(CancelledTask);
+	}
+}
\ No newline at end of file
diff --git a/Engine/Plugins/Lumin/MagicLeap/Source/MagicLeap/Private/CameraCaptureRunnable.h b/Engine/Plugins/Lumin/MagicLeap/Source/MagicLeap/Private/CameraCaptureRunnable.h
--- a/Engine/Plugins/Lumin/MagicLeap/Source/MagicLeap/Private/CameraCaptureRunnable.h
+++ b/Engine/Plugins/Lumin/MagicLeap/Source/MagicLeap/Private/CameraCaptureRunnable.h
@@ -55,17 +55,25 @@ ML_INCLUDES_END
 
 enum class CaptureMsgType : uint32
 {
+	Invalid,
 	Request,
 	Response,
 	Log
 };
 
+enum class CaptureTaskType : uint32
+{
+	None,
+	ImageToFile,
+	ImageToTexture,
+	StartVideoToFile,
+	StopVideoToFile,
+};
+
 struct FCaptureMessage
 {
 	CaptureMsgType Type;
-#if WITH_MLSDK
-	MLCameraCaptureType CaptureType;
-#endif //WITH_MLSDK
+	CaptureTaskType CaptureType;
 	FString Log;
 	FString FilePath;
 	bool Success;
@@ -74,16 +82,20 @@ struct FCaptureMessage
 	class FCameraCaptureImpl* Requester;
 
 	FCaptureMessage()
-		: Type(CaptureMsgType::Request)
-#if WITH_MLSDK
-		, CaptureType(MLCameraCaptureType_Image)
-#endif //WITH_MLSDK
-		, Log("")
-		, Success(false)
-		, Texture(nullptr)
-		, Duration(0.0f)
-		, Requester(nullptr)
-	{}
+	{
+		Clear();
+	}
+
+	void Clear()
+	{
+		Type = CaptureMsgType::Invalid;
+		CaptureType = CaptureTaskType::None;
+		Log = "";
+		Success = false;
+		Texture = nullptr;
+		Duration = 0.0f;
+		Requester = nullptr;
+	}
 };
 
 class FCameraCaptureRunnable : public FRunnable, public MagicLeap::IAppEventHandler
@@ -93,22 +105,11 @@ public:
 	virtual ~FCameraCaptureRunnable();
 
 	uint32 Run() override;
-
 	void Stop() override;
-
-	void BeginCapture();
-
-	void EndCapture(bool InSuccess);
-
-	void Log(const FString& Info);
-
-	void ProcessCaptureMessage(const FCaptureMessage& InMsg);
-
 	void OnAppPause() override;
-
 	void OnAppResume() override;
-
 	void OnAppShutDown() override;
+	void ProcessCaptureMessage(const FCaptureMessage& InMsg);
 
 	/** Internal thread this instance is running on */
 	FRunnableThread* Thread;
@@ -116,14 +117,10 @@ public:
 	TQueue<FCaptureMessage, EQueueMode::Spsc> IncomingMessages;
 	TQueue<FCaptureMessage, EQueueMode::Spsc> OutgoingMessages;
 	FEvent* Semaphore;
-	FCaptureMessage CurrentMessage;
 	bool bCameraConnected;
-	const float RetryConnectWaitTime;
 
 #if WITH_MLSDK
 	MLCameraDeviceStatusCallbacks DeviceStatusCallbacks;
-	MLCameraCaptureCallbacks CaptureCallbacks;
-	MLCameraOutput* CameraOutput;
 #endif //WITH_MLSDK
 
 	const FString ImgExtension;
@@ -131,12 +128,23 @@ public:
 	FString UniqueFileName;
 	TSharedPtr<IImageWrapper> ImageWrapper;
 	FThreadSafeBool bPaused;
+	FCaptureMessage CurrentTask;
 	static FThreadSafeCounter64 PreviewHandle;
 
 private:
 #if WITH_MLSDK
 	static void OnPreviewBufferAvailable(MLHandle Output, void *Data);
 #endif //WITH_MLSDK
+
+	bool TryConnect();
+	bool DoCaptureTask();
+	bool CaptureImageToFile();
+	bool CaptureImageToTexture();
+	bool StartRecordingVideo();
+	bool StopRecordingVideo();
+	void Pause();
+	void Log(const FString& Info);
+	void CancelIncomingTasks();
 };
 
 DECLARE_LOG_CATEGORY_EXTERN(LogCameraCaptureRunnable, Verbose, All);
\ No newline at end of file
diff --git a/Engine/Plugins/Lumin/MagicLeap/Source/MagicLeap/Private/IMagicLeapModule.h b/Engine/Plugins/Lumin/MagicLeap/Source/MagicLeap/Private/IMagicLeapModule.h
new file mode 100644
--- /dev/null
+++ b/Engine/Plugins/Lumin/MagicLeap/Source/MagicLeap/Private/IMagicLeapModule.h
@@ -0,0 +1,60 @@
+// Copyright 1998-2018 Epic Games, Inc. All Rights Reserved.
+
+#pragma once
+
+//#include "CoreMinimal.h"
+#include "AppFramework.h"
+
+/**
+ * This is the base module interface which all magic leap modules should inherit from.
+ * The purpose of this class is to allow the system to manage the order in which magic
+ * leap plugins are loaded and unloaded.
+ */
+class IMagicLeapModule
+{
+public:
+	/** Registers the module with the AppFramework instance.*/
+	IMagicLeapModule(FName InName) : Name(InName)
+	{
+		FAppFramework::RegisterMagicLeapModule(this);
+	}
+
+	/** Unregisters the module with the Appframework instance.*/
+	virtual ~IMagicLeapModule()
+	{
+		FAppFramework::UnregisterMagicLeapModule(this);
+	}
+
+	/** Returns the name of the module.*/
+	FName GetName() const
+	{
+		return Name;
+	}
+
+	/** Override to perform order-dependant initialization of your plugin.*/
+	virtual void Enable() {}
+
+	/** Override to perform order-dependant cleanup of your plugin.*/
+	virtual void Disable() {}
+
+protected:
+	const FName Name;
+};
+
+#define ENABLE_MAGIC_LEAP_MODULE(Name) \
+{ \
+	IMagicLeapModule* MagicLeapModule = FAppFramework::GetMagicLeapModule(Name); \
+	if (nullptr != MagicLeapModule) \
+	{ \
+		MagicLeapModule->Enable(); \
+	} \
+}
+
+#define DISABLE_MAGIC_LEAP_MODULE(Name) \
+{ \
+	IMagicLeapModule* MagicLeapModule = FAppFramework::GetMagicLeapModule(Name); \
+	if (nullptr != MagicLeapModule) \
+	{ \
+		MagicLeapModule->Disable(); \
+	} \
+}
\ No newline at end of file
