Updated ML integration to MLSDK 0.19

diff --git a/Engine/Plugins/Lumin/MagicLeap/Source/MagicLeap/Private/ImageTrackerComponent.cpp b/Engine/Plugins/Lumin/MagicLeap/Source/MagicLeap/Private/ImageTrackerComponent.cpp
--- a/Engine/Plugins/Lumin/MagicLeap/Source/MagicLeap/Private/ImageTrackerComponent.cpp
+++ b/Engine/Plugins/Lumin/MagicLeap/Source/MagicLeap/Private/ImageTrackerComponent.cpp
@@ -250,6 +250,7 @@ class FImageTrackerImpl : public FRunnable, public MagicLeap::IAppEventHandler
 public:
 	FImageTrackerImpl()
 		: bHasTarget(false)
+		, bIsTracking(false)
 #if WITH_EDITOR
 		, TextureBeforeEdit(nullptr)
 #endif
@@ -258,7 +259,6 @@ public:
 #endif //WITH_MLSDK
 		, Thread(nullptr)
 		, StopTaskCounter(0)
-		, bIsTracking(false)
 	{
 #if WITH_MLSDK
 		OldTrackingStatus.status = MLImageTrackerTargetStatus_Ensure32Bits;
@@ -297,6 +297,7 @@ public:
 			{
 				ImageTracker = FImageTrackerEngineInterface::Get();
 			}
+
 			if (MLHandleIsValid(ImageTracker.Pin()->GetHandle()))
 			{
 				if (IncomingMessages.Dequeue(CurrentMessage))
@@ -321,10 +322,12 @@ public:
 		}
 
 		bHasTarget = true;
+		bIsTracking = false;
 		MLImageTrackerTargetSettings TargetSettings;
 		TargetSettings.longer_dimension = InLongerDimension;
 		TargetSettings.is_stationary = bInIsStationary;
 		FTrackerMessage CreateTargetMsg;
+		CreateTargetMsg.Requester = this;
 		CreateTargetMsg.TaskType = FTrackerMessage::TaskType::TryCreateTarget;
 		CreateTargetMsg.TargetName = InName;
 		CreateTargetMsg.TargetSettings = TargetSettings;
@@ -343,10 +346,25 @@ public:
 		IncomingMessages.Enqueue(CreateTargetMsg);
 #endif //WITH_MLSDK
 	}
+
+	bool TryGetResult(FTrackerMessage& OutMsg)
+	{
+#if PLATFORM_LUMIN
+		if (OutgoingMessages.Peek(OutMsg))
+		{
+			if (OutMsg.Requester == this)
+			{
+				OutgoingMessages.Pop();
+				return true;
+			}
+		}
+#endif //PLATFORM_LUMIN
+		return false;
+	}
 	
 public:
-	TQueue<FTrackerMessage, EQueueMode::Spsc> OutgoingMessages;
 	bool bHasTarget;
+	bool bIsTracking;
 #if WITH_EDITOR
 	UTexture2D* TextureBeforeEdit;
 #endif
@@ -355,14 +373,13 @@ public:
 	MLImageTrackerTargetStaticData Data;
 	MLImageTrackerTargetResult OldTrackingStatus;
 #endif //WITH_MLSDK
-
-public:
 	TWeakPtr<FImageTrackerEngineInterface, ESPMode::ThreadSafe> ImageTracker;
 	FRunnableThread* Thread;
 	FThreadSafeCounter StopTaskCounter;
 	TQueue<FTrackerMessage, EQueueMode::Spsc> IncomingMessages;
+	TQueue<FTrackerMessage, EQueueMode::Spsc> OutgoingMessages;
 	FTrackerMessage CurrentMessage;
-	bool bIsTracking;
+	FCriticalSection DataMutex;
 
 	void DoTasks()
 	{
@@ -377,35 +394,48 @@ public:
 	void SetTarget()
 	{
 #if WITH_MLSDK
-		if (!MLHandleIsValid(Target))
+		FScopeLock Lock(&DataMutex);
+
+		if (MLHandleIsValid(Target))
 		{
-			UE_LOG(LogMagicLeap, Warning, TEXT("SetTarget for %s"), *CurrentMessage.TargetName);
-			CurrentMessage.TargetSettings.name = TCHAR_TO_UTF8(*CurrentMessage.TargetName);
-			FTexture2DMipMap& Mip = CurrentMessage.TargetImageTexture->PlatformData->Mips[0];
-			const unsigned char* PixelData = static_cast<const unsigned char*>(Mip.BulkData.Lock(LOCK_READ_ONLY));
-			checkf(ImageTracker.Pin().IsValid(), TEXT("[FImageTrackerImpl] ImageTracker weak pointer is invalid!"));
-			MLImageTrackerAddTargetFromArray(ImageTracker.Pin()->GetHandle(), &CurrentMessage.TargetSettings, PixelData, CurrentMessage.TargetImageTexture->GetSurfaceWidth(), CurrentMessage.TargetImageTexture->GetSurfaceHeight(), MLImageTrackerImageFormat_RGBA, &Target);
-			Mip.BulkData.Unlock();
+			MLImageTrackerRemoveTarget(ImageTracker.Pin()->GetHandle(), Target);
+		}
 
-			if (!MLHandleIsValid(Target))
-			{
-				UE_LOG(LogMagicLeap, Error, TEXT("[FImageTrackerImpl] Could not create Image Target."));
-				return;
-			}
+		UE_LOG(LogMagicLeap, Warning, TEXT("SetTarget for %s"), *CurrentMessage.TargetName);
+		FTrackerMessage TargetCreateMsg;
+		TargetCreateMsg.Requester = CurrentMessage.Requester;
 
-			// [3] Cache all the static data for this target.
-			checkf(ImageTracker.Pin().IsValid(), TEXT("[FImageTrackerImpl] ImageTracker weak pointer is invalid!"));
-			MLResult Result = MLImageTrackerGetTargetStaticData(ImageTracker.Pin()->GetHandle(), Target, &Data);
-			if (Result != MLResult_Ok)
-			{
-				UE_LOG(LogMagicLeap, Error, TEXT("[FImageTrackerImpl] Could not get the static data for the Image Target."));
-				return;
-			}
+		CurrentMessage.TargetSettings.name = TCHAR_TO_UTF8(*CurrentMessage.TargetName);
+		FTexture2DMipMap& Mip = CurrentMessage.TargetImageTexture->PlatformData->Mips[0];
+		const unsigned char* PixelData = static_cast<const unsigned char*>(Mip.BulkData.Lock(LOCK_READ_ONLY));
+		checkf(ImageTracker.Pin().IsValid(), TEXT("[FImageTrackerImpl] ImageTracker weak pointer is invalid!"));
+		MLImageTrackerAddTargetFromArray(ImageTracker.Pin()->GetHandle(), &CurrentMessage.TargetSettings, PixelData, CurrentMessage.TargetImageTexture->GetSurfaceWidth(), CurrentMessage.TargetImageTexture->GetSurfaceHeight(), MLImageTrackerImageFormat_RGBA, &Target);
+		Mip.BulkData.Unlock();
 
-			UE_LOG(LogMagicLeap, Warning, TEXT("SetTarget successfully set for %s"), *CurrentMessage.TargetName);			
+		if (!MLHandleIsValid(Target))
+		{
+			UE_LOG(LogMagicLeap, Error, TEXT("[FImageTrackerImpl] Could not create Image Target."));
+			TargetCreateMsg.TaskType = FTrackerMessage::TaskType::TargetCreateFailed;
+			OutgoingMessages.Enqueue(TargetCreateMsg);
+			return;
+		}
 
-			bIsTracking = true;
+		// [3] Cache all the static data for this target.
+		checkf(ImageTracker.Pin().IsValid(), TEXT("[FImageTrackerImpl] ImageTracker weak pointer is invalid!"));
+		MLResult Result = MLImageTrackerGetTargetStaticData(ImageTracker.Pin()->GetHandle(), Target, &Data);
+		if (Result != MLResult_Ok)
+		{
+			UE_LOG(LogMagicLeap, Error, TEXT("[FImageTrackerImpl] Could not get the static data for the Image Target."));
+			TargetCreateMsg.TaskType = FTrackerMessage::TaskType::TargetCreateFailed;
+			OutgoingMessages.Enqueue(TargetCreateMsg);
+			return;
 		}
+
+		UE_LOG(LogMagicLeap, Log, TEXT("SetTarget successfully set for %s"), *CurrentMessage.TargetName);
+		TargetCreateMsg.TaskType = FTrackerMessage::TaskType::TargetCreateSucceeded;
+		TargetCreateMsg.Target = Target;
+		TargetCreateMsg.Data = Data;
+		OutgoingMessages.Enqueue(TargetCreateMsg);
 #endif //WITH_MLSDK
 	}
 
@@ -425,7 +455,6 @@ public:
 
 UImageTrackerComponent::UImageTrackerComponent()
 	: Impl(new FImageTrackerImpl())
-	, bTick(true)
 {
 	// Make sure this component ticks
 	PrimaryComponentTick.bCanEverTick = true;
@@ -440,7 +469,10 @@ UImageTrackerComponent::~UImageTrackerComponent()
 #if PLATFORM_LUMIN
 	if (Impl)
 	{
-		Impl->AsyncDestroy();
+		if (!Impl->AsyncDestroy())
+		{
+			delete Impl;
+		}
 		Impl = nullptr;
 	}
 #else
@@ -454,6 +486,11 @@ void UImageTrackerComponent::TickComponent(float DeltaTime, enum ELevelTick Tick
 #if WITH_MLSDK && PLATFORM_LUMIN
 	Super::TickComponent(DeltaTime, TickType, ThisTickFunction);
 
+	if (!(IMagicLeapPlugin::Get().IsMagicLeapHMDValid()))
+	{
+		return;
+	}
+
 	const FAppFramework& AppFramework = StaticCastSharedPtr<FMagicLeapHMD>(GEngine->XRSystem)->GetAppFrameworkConst();
 
 	if (!AppFramework.IsInitialized())
@@ -471,85 +508,137 @@ void UImageTrackerComponent::TickComponent(float DeltaTime, enum ELevelTick Tick
 	if (TargetImageTexture->GetPixelFormat() != EPixelFormat::PF_R8G8B8A8 && TargetImageTexture->GetPixelFormat() != EPixelFormat::PF_B8G8R8A8)
 	{
 		UE_LOG(LogMagicLeap, Error, TEXT("[UImageTrackerComponent] ImageTracker: Unsupported pixel format encountered!"));
-		bTick = false;
+		return;
 	}
 
-	if (!bTick) return;
-
 	if (!Impl->bHasTarget)
 	{
 		if (Name.Len() == 0) Name = GetName();
 		Impl->SetTargetAsync(Name, bIsStationary, LongerDimension / AppFramework.GetWorldToMetersScale(), TargetImageTexture);
 	}
 
+	FTrackerMessage TrackingResult;
+	if (Impl->TryGetResult(TrackingResult))
+	{
+		if (TrackingResult.TaskType == FTrackerMessage::TaskType::TargetCreateSucceeded)
+		{
+			OnSetImageTargetSucceeded.Broadcast();
+			Impl->bIsTracking = true;
+		}
+		else if(TrackingResult.TaskType == FTrackerMessage::TaskType::TargetCreateFailed)
+		{
+			OnSetImageTargetFailed.Broadcast();
+			Impl->bIsTracking = false;
+		}
+	}
+
 	if (Impl->bIsTracking)
 	{
 		checkf(Impl->ImageTracker.Pin().IsValid(), TEXT("[FImageTrackerImpl] ImageTracker weak pointer is invalid!"));
-
 		MLImageTrackerTargetResult TrackingStatus;
-		if (MLImageTrackerGetTargetResult(Impl->ImageTracker.Pin()->GetHandle(), Impl->Target, &TrackingStatus) != MLResult_Ok)
 		{
-			TrackingStatus.status = MLImageTrackerTargetStatus_NotTracked;
+			FScopeLock Lock(&Impl->DataMutex);
+			MLResult Result = MLImageTrackerGetTargetResult(Impl->ImageTracker.Pin()->GetHandle(), Impl->Target, &TrackingStatus);
+			if (Result != MLResult_Ok)
+			{
+				UE_LOG(LogMagicLeap, Warning, TEXT("MLImageTrackerGetTargetResult failed due to error %s."), UTF8_TO_TCHAR(MLGetResultString(Result)));
+				TrackingStatus.status = MLImageTrackerTargetStatus_NotTracked;
+			}
 		}
 
 		if (TrackingStatus.status == MLImageTrackerTargetStatus_NotTracked)
+		{
+			if (Impl->OldTrackingStatus.status != MLImageTrackerTargetStatus_NotTracked)
 			{
-				if (Impl->OldTrackingStatus.status != MLImageTrackerTargetStatus_NotTracked)
-				{
-					OnImageTargetLost.Broadcast();
-				}
+				OnImageTargetLost.Broadcast();
 			}
-			else
+		}
+		else
+		{
+			EFailReason FailReason = EFailReason::None;
+			FTransform Pose = FTransform::Identity;
+			bool bHasTransform = false;
 			{
-				EFailReason FailReason = EFailReason::None;
-				FTransform Pose = FTransform::Identity;
-				if (AppFramework.GetTransform(Impl->Data.coord_frame_target, Pose, FailReason))
-				{
-				Pose.SetRotation(MagicLeap::ToUERotator(Pose.GetRotation()));
+				FScopeLock Lock(&Impl->DataMutex);
+				bHasTransform = AppFramework.GetTransform(Impl->Data.coord_frame_target, Pose, FailReason);
+			}
+			if (bHasTransform)
+			{
+				Pose.ConcatenateRotation(FQuat(FVector(0, 0, 1), PI));
 				if (TrackingStatus.status == MLImageTrackerTargetStatus_Unreliable)
-					{
-						FVector LastTrackedLocation = GetComponentLocation();
-						FRotator LastTrackedRotation = GetComponentRotation();
-						if (bUseUnreliablePose)
-						{
-							this->SetRelativeLocationAndRotation(Pose.GetTranslation(), Pose.Rotator());
-						}
-						// Developer can choose whether to use this unreliable pose or not.
-						OnImageTargetUnreliableTracking.Broadcast(LastTrackedLocation, LastTrackedRotation, Pose.GetTranslation(), Pose.Rotator());
-					}
-					else
+				{
+					FVector LastTrackedLocation = GetComponentLocation();
+					FRotator LastTrackedRotation = GetComponentRotation();
+					if (bUseUnreliablePose)
 					{
 						this->SetRelativeLocationAndRotation(Pose.GetTranslation(), Pose.Rotator());
-						if (Impl->OldTrackingStatus.status != MLImageTrackerTargetStatus_Tracked)
-						{
-							OnImageTargetFound.Broadcast();
-						}
 					}
+					// Developer can choose whether to use this unreliable pose or not.
+					OnImageTargetUnreliableTracking.Broadcast(LastTrackedLocation, LastTrackedRotation, Pose.GetTranslation(), Pose.Rotator());
 				}
 				else
 				{
-					if (FailReason == EFailReason::NaNsInTransform)
+					this->SetRelativeLocationAndRotation(Pose.GetTranslation(), Pose.Rotator());
+					if (Impl->OldTrackingStatus.status != MLImageTrackerTargetStatus_Tracked)
 					{
-						UE_LOG(LogMagicLeap, Error, TEXT("[UImageTrackerComponent] NaNs in image tracker target transform."));
+						OnImageTargetFound.Broadcast();
 					}
+				}
+			}
+			else
+			{
+				if (FailReason == EFailReason::NaNsInTransform)
+				{
+					UE_LOG(LogMagicLeap, Error, TEXT("[UImageTrackerComponent] NaNs in image tracker target transform."));
+				}
 				TrackingStatus.status = MLImageTrackerTargetStatus_NotTracked;
-					if (Impl->OldTrackingStatus.status != MLImageTrackerTargetStatus_NotTracked)
-					{
-						OnImageTargetLost.Broadcast();
-					}
+				if (Impl->OldTrackingStatus.status != MLImageTrackerTargetStatus_NotTracked)
+				{
+					OnImageTargetLost.Broadcast();
 				}
 			}
+		}
 
 		Impl->OldTrackingStatus = TrackingStatus;
 	}
 #endif //WITH_MLSDK
 }
 
-void UImageTrackerComponent::SetTargetAsync(UTexture2D* ImageTarget)
+bool UImageTrackerComponent::SetTargetAsync(UTexture2D* ImageTarget)
 {
+	if (ImageTarget == nullptr)
+	{
+		UE_LOG(LogMagicLeap, Warning, TEXT("[UImageTrackerComponent] ImageTarget is NULL!."));
+		return false;
+	}
+
+	const FAppFramework& AppFramework = StaticCastSharedPtr<FMagicLeapHMD>(GEngine->XRSystem)->GetAppFrameworkConst();
+
+	if (!AppFramework.IsInitialized())
+	{
+		UE_LOG(LogMagicLeap, Warning, TEXT("[UImageTrackerComponent] AppFramework not initialized."));
+		return false;
+	}
 
+	if (ImageTarget && !(ImageTarget->GetPixelFormat() == EPixelFormat::PF_R8G8B8A8 || ImageTarget->GetPixelFormat() == EPixelFormat::PF_B8G8R8A8))
+	{
+		UE_LOG(LogMagicLeap, Error, TEXT("[UImageTrackerComponent] Cannot set texture %s as it uses an invalid pixel format!  Valid formats are R8B8G8A8 or B8G8R8A8"), *ImageTarget->GetName());
+		return false;
+	}
+
+	if (ImageTarget == TargetImageTexture)
+	{
+		UE_LOG(LogMagicLeap, Error, TEXT("[UImageTrackerComponent] Skipped setting %s as it is already being used as the current image target"), *ImageTarget->GetName());
+		return false;
+	}
+
+	TargetImageTexture = ImageTarget;
+	Impl->SetTargetAsync(Name, bIsStationary, LongerDimension / AppFramework.GetWorldToMetersScale(), TargetImageTexture);
+	return true;
 }
 
+
+
 #if WITH_EDITOR
 void UImageTrackerComponent::PreEditChange(UProperty* PropertyAboutToChange)
 {
@@ -569,8 +658,8 @@ void UImageTrackerComponent::PostEditChangeProperty(FPropertyChangedEvent& Prope
 	{
 		if (TargetImageTexture && !(TargetImageTexture->GetPixelFormat() == EPixelFormat::PF_R8G8B8A8 || TargetImageTexture->GetPixelFormat() == EPixelFormat::PF_B8G8R8A8))
 		{
-			TargetImageTexture = Impl->TextureBeforeEdit;
 			UE_LOG(LogMagicLeap, Error, TEXT("[UImageTrackerComponent] Cannot set texture %s as it uses an invalid pixel format!  Valid formats are R8B8G8A8 or B8G8R8A8"), *TargetImageTexture->GetName());
+			TargetImageTexture = Impl->TextureBeforeEdit;
 		}
 	}
 
diff --git a/Engine/Plugins/Lumin/MagicLeap/Source/MagicLeap/Private/ImageTrackerRunnable.h b/Engine/Plugins/Lumin/MagicLeap/Source/MagicLeap/Private/ImageTrackerRunnable.h
--- a/Engine/Plugins/Lumin/MagicLeap/Source/MagicLeap/Private/ImageTrackerRunnable.h
+++ b/Engine/Plugins/Lumin/MagicLeap/Source/MagicLeap/Private/ImageTrackerRunnable.h
@@ -44,10 +44,6 @@
 #if WITH_MLSDK
 ML_INCLUDES_START
 #include <ml_image_tracking.h>
-//#include <ml_perception.h>
-//#include <ml_snapshot.h>
-//#include <ml_head_tracking.h>
-//#include <ml_coordinate_frame_uid.h>
 ML_INCLUDES_END
 #endif //WITH_MLSDK
 
diff --git a/Engine/Plugins/Lumin/MagicLeap/Source/MagicLeap/Private/MagicLeapARPinComponent.cpp b/Engine/Plugins/Lumin/MagicLeap/Source/MagicLeap/Private/MagicLeapARPinComponent.cpp
--- a/Engine/Plugins/Lumin/MagicLeap/Source/MagicLeap/Private/MagicLeapARPinComponent.cpp
+++ b/Engine/Plugins/Lumin/MagicLeap/Source/MagicLeap/Private/MagicLeapARPinComponent.cpp
@@ -445,6 +445,8 @@ void UMagicLeapARPinComponent::TickComponent(float DeltaTime, enum ELevelTick Ti
 		return;
 	}
 
+	bool bDidPin = false;
+
 	if (PinnedSceneComponent != nullptr)
 	{
 		if (!bPinned)
@@ -483,10 +485,7 @@ void UMagicLeapARPinComponent::TickComponent(float DeltaTime, enum ELevelTick Ti
 				}
 			}
 
-			if (bPinned)
-			{
-				OnPersistentEntityPinned.Broadcast(bDataRestored);
-			}
+			bDidPin = bPinned;
 		}
 
 		if (bPinned)
@@ -503,6 +502,11 @@ void UMagicLeapARPinComponent::TickComponent(float DeltaTime, enum ELevelTick Ti
 				NewComponentWorldTransform = FTransform(NewComponentWorldTransformMatrix);
 
 				PinnedSceneComponent->SetWorldLocationAndRotation(NewComponentWorldTransform.GetLocation(), NewComponentWorldTransform.Rotator());
+
+				if (bDidPin)
+				{
+					OnPersistentEntityPinned.Broadcast(bDataRestored);
+				}
 			}
 		}
 	}
diff --git a/Engine/Plugins/Lumin/MagicLeap/Source/MagicLeap/Private/MagicLeapCustomPresent.cpp b/Engine/Plugins/Lumin/MagicLeap/Source/MagicLeap/Private/MagicLeapCustomPresent.cpp
--- a/Engine/Plugins/Lumin/MagicLeap/Source/MagicLeap/Private/MagicLeapCustomPresent.cpp
+++ b/Engine/Plugins/Lumin/MagicLeap/Source/MagicLeap/Private/MagicLeapCustomPresent.cpp
@@ -329,7 +329,12 @@ void FMagicLeapCustomPresentOpenGL::FinishRendering()
 
 		glBindFramebuffer(GL_READ_FRAMEBUFFER, Framebuffers[0]);
 		glBindFramebuffer(GL_DRAW_FRAMEBUFFER, Framebuffers[1]);
-		FOpenGL::BlitFramebuffer(0, 0, InternalTextureDims.X / 2, InternalTextureDims.Y, 0, 0, vp_width, vp_height, GL_COLOR_BUFFER_BIT, GL_NEAREST);
+
+		const bool bShouldFlipVertically = !IsES2Platform(GMaxRHIShaderPlatform);
+
+		bShouldFlipVertically ?
+			FOpenGL::BlitFramebuffer(0, 0, InternalTextureDims.X / 2, InternalTextureDims.Y, 0, vp_height, vp_width, 0, GL_COLOR_BUFFER_BIT, GL_NEAREST) :
+			FOpenGL::BlitFramebuffer(0, 0, InternalTextureDims.X / 2, InternalTextureDims.Y, 0, 0, vp_width, vp_height, GL_COLOR_BUFFER_BIT, GL_NEAREST);
 
 		MLResult Result = MLGraphicsSignalSyncObjectGL(Plugin->GraphicsClient, vp_array.virtual_cameras[0].sync_object);
 		if (Result != MLResult_Ok)
@@ -338,7 +343,11 @@ void FMagicLeapCustomPresentOpenGL::FinishRendering()
 		}
 
 		FOpenGL::FramebufferTextureLayer(GL_FRAMEBUFFER, GL_COLOR_ATTACHMENT0, vp_array.color_id, 0, 1);
-		FOpenGL::BlitFramebuffer(InternalTextureDims.X / 2, 0, InternalTextureDims.X, InternalTextureDims.Y, 0, 0, vp_width, vp_height, GL_COLOR_BUFFER_BIT, GL_NEAREST);
+
+		bShouldFlipVertically ?
+			FOpenGL::BlitFramebuffer(InternalTextureDims.X / 2, 0, InternalTextureDims.X, InternalTextureDims.Y, 0, vp_height, vp_width, 0, GL_COLOR_BUFFER_BIT, GL_NEAREST) :
+			FOpenGL::BlitFramebuffer(InternalTextureDims.X / 2, 0, InternalTextureDims.X, InternalTextureDims.Y, 0, 0, vp_width, vp_height, GL_COLOR_BUFFER_BIT, GL_NEAREST);
+
 		Result = MLGraphicsSignalSyncObjectGL(Plugin->GraphicsClient, vp_array.virtual_cameras[1].sync_object);
 		if (Result != MLResult_Ok)
 		{
@@ -463,38 +472,38 @@ void FMagicLeapCustomPresentVulkan::BeginRendering()
 {
 #if WITH_MLSDK
 	check(IsInRenderingThread());
-	FTrackingFrame& frame = Plugin->GetCurrentFrameMutable();
-	if (bCustomPresentIsSet)
+
+	ExecuteOnRHIThread([this]()
 	{
-		// TODO [Blake] : Need to see if we can use this newer matrix and override the view
-		// projection matrix (since they query GetStereoProjectionMatrix on the main thread)
-		MLGraphicsFrameParams camera_params;
-		MLResult InitResult = MLGraphicsInitFrameParams(&camera_params);
-		if (InitResult != MLResult_Ok)
+		// Always use RHITrackingFrame here, which is then copied to the RenderTrackingFrame. 
+		FTrackingFrame& RHIframe = Plugin->RHITrackingFrame;
+		if (bCustomPresentIsSet)
 		{
-			UE_LOG(LogMagicLeap, Error, TEXT("MLGraphicsInitFrameParams failed with status %d"), InitResult);
-		}
-		camera_params.projection_type = MLGraphicsProjectionType_ReversedInfiniteZ;
-		camera_params.surface_scale = 1.0f;
-		camera_params.protected_surface = false;
-		GConfig->GetBool(TEXT("/Script/LuminRuntimeSettings.LuminRuntimeSettings"), TEXT("bProtectedContent"), camera_params.protected_surface, GEngineIni);
+			// TODO [Blake] : Need to see if we can use this newer matrix and override the view
+			// projection matrix (since they query GetStereoProjectionMatrix on the main thread)
+			MLGraphicsFrameParams camera_params;
+			MLResult InitResult = MLGraphicsInitFrameParams(&camera_params);
+			if (InitResult != MLResult_Ok)
+			{
+				UE_LOG(LogMagicLeap, Error, TEXT("MLGraphicsInitFrameParams failed with status %d"), InitResult);
+			}
+			camera_params.projection_type = MLGraphicsProjectionType_UnsignedZ;
+			camera_params.surface_scale = 1.0f;
+			camera_params.protected_surface = false;
+			GConfig->GetBool(TEXT("/Script/LuminRuntimeSettings.LuminRuntimeSettings"), TEXT("bProtectedContent"), camera_params.protected_surface, GEngineIni);
 
-		// The near clipping plane is expected in meters despite what is documented in the header.
-		camera_params.near_clip = GNearClippingPlane / frame.WorldToMetersScale;
-		camera_params.far_clip = frame.FarClippingPlane / frame.WorldToMetersScale;
+			// The near clipping plane is expected in meters despite what is documented in the header.
+			camera_params.near_clip = GNearClippingPlane / RHIframe.WorldToMetersScale;
+			camera_params.far_clip = RHIframe.FarClippingPlane / RHIframe.WorldToMetersScale;
 
-		// The focus distance is expected in meters despite what is documented in the header.
-		// Only focus distance equaling 1 engine unit seems to work on board without wearable and on desktop.
+			// The focus distance is expected in meters despite what is documented in the header.
+			// Only focus distance equaling 1 engine unit seems to work on board without wearable and on desktop.
 #if PLATFORM_LUMIN
-		camera_params.focus_distance = frame.FocusDistance / frame.WorldToMetersScale;
+			camera_params.focus_distance = RHIframe.FocusDistance / RHIframe.WorldToMetersScale;
 #else
-		camera_params.focus_distance = 1.0f;
+			camera_params.focus_distance = 1.0f;
 #endif
 
-	  ExecuteOnRHIThread_DoNotWait([this, camera_params]()
-	  {
-		  FTrackingFrame& RHIframe = Plugin->GetCurrentFrameMutable();
-
 #if BEGIN_END_FRAME_BALANCE_HACK
 			if (BalanceCounter != 0)
 			{
@@ -509,39 +518,41 @@ void FMagicLeapCustomPresentVulkan::BeginRendering()
 			}
 #endif
 
-		  MLResult Result = MLGraphicsBeginFrame(Plugin->GraphicsClient, &camera_params, &RHIframe.Handle, &RHIframe.RenderInfoArray);
-		  RHIframe.bBeginFrameSucceeded = (Result == MLResult_Ok);
-		  if (RHIframe.bBeginFrameSucceeded)
-		  {
+			MLResult Result = MLGraphicsBeginFrame(Plugin->GraphicsClient, &camera_params, &RHIframe.Handle, &RHIframe.RenderInfoArray);
+			RHIframe.bBeginFrameSucceeded = (Result == MLResult_Ok);
+			if (RHIframe.bBeginFrameSucceeded)
+			{
 #if BEGIN_END_FRAME_BALANCE_HACK
 				++BalanceCounter;
 				BalancePrevFrameHandle = RHIframe.Handle;
 				BalanceHandles[0] = RHIframe.RenderInfoArray.virtual_cameras[0].sync_object;
 				BalanceHandles[1] = RHIframe.RenderInfoArray.virtual_cameras[1].sync_object;
 #endif //BEGIN_END_FRAME_BALANCE_HACK
+
 			  /* Convert eye extents from Graphics Projection Model to Unreal Projection Model */
 			  // Unreal expects the projection matrix to be in centimeters and uses it for various purposes
 			  // such as bounding volume calculations for lights in the shadow algorithm.
 			  // We're overwriting the near value to match the units of unreal here instead of using the units of the SDK
-			  for (uint32_t eye = 0; eye < RHIframe.RenderInfoArray.num_virtual_cameras; ++eye)
-			  {
-				  RHIframe.RenderInfoArray.virtual_cameras[eye].projection.matrix_colmajor[10] = 0.0f; // Model change hack
-				  RHIframe.RenderInfoArray.virtual_cameras[eye].projection.matrix_colmajor[11] = -1.0f; // Model change hack
-				  RHIframe.RenderInfoArray.virtual_cameras[eye].projection.matrix_colmajor[14] = GNearClippingPlane; // Model change hack
-			  }
-		  }
-		  else
-		  {
-			  if (Result != MLResult_Timeout)
-			  {
+				for (uint32_t eye = 0; eye < RHIframe.RenderInfoArray.num_virtual_cameras; ++eye)
+				{
+					RHIframe.RenderInfoArray.virtual_cameras[eye].projection.matrix_colmajor[10] = 0.0f; // Model change hack
+					RHIframe.RenderInfoArray.virtual_cameras[eye].projection.matrix_colmajor[11] = -1.0f; // Model change hack
+					RHIframe.RenderInfoArray.virtual_cameras[eye].projection.matrix_colmajor[14] = GNearClippingPlane; // Model change hack
+				}
+			}
+			else
+			{
+				if (Result != MLResult_Timeout)
+				{
 				UE_LOG(LogMagicLeap, Error, TEXT("MLGraphicsBeginFrame failed with status %d"), Result);
-			  }
-			  // TODO: See if this is only needed for ZI.
-			  RHIframe.Handle = ML_INVALID_HANDLE;
-			  MagicLeap::ResetVirtualCameraInfoArray(RHIframe.RenderInfoArray);
-		  }
-	  });
-  }
+				}
+				// TODO: See if this is only needed for ZI.
+				RHIframe.Handle = ML_INVALID_HANDLE;
+				MagicLeap::ResetVirtualCameraInfoArray(RHIframe.RenderInfoArray);
+			}
+			Plugin->InitializeRenderFrameFromRHIFrame();
+		}
+	});
 #endif // WITH_MLSDK
 }
 
@@ -567,9 +578,13 @@ void FMagicLeapCustomPresentVulkan::FinishRendering()
 			// Alias the render target with an srgb image description for proper color space output.
 			if (RenderTargetTextureAllocation != VK_NULL_HANDLE && LastAliasedRenderTarget != RenderTargetTexture)
 			{
-				// TODO: If RenderTargetTextureSRGB is non-null, we're leaking the previous handle here. Also leaking on shutdown.
-				RenderTargetTextureSRGB = (VkImage)FMagicLeapHelperVulkan::AliasImageSRGB((uint64)RenderTargetTextureAllocation, RenderTargetTextureAllocationOffset, vp_width * 2, vp_height);
-				check(RenderTargetTextureSRGB != VK_NULL_HANDLE);
+				// SDKUNREAL-1135: ML remote image is corrupted on AMD hardware
+				if (!IsRHIDeviceAMD())
+				{
+					// TODO: If RenderTargetTextureSRGB is non-null, we're leaking the previous handle here. Also leaking on shutdown.
+					RenderTargetTextureSRGB = reinterpret_cast<void *>(FMagicLeapHelperVulkan::AliasImageSRGB((uint64)RenderTargetTextureAllocation, RenderTargetTextureAllocationOffset, vp_width * 2, vp_height));
+					check(RenderTargetTextureSRGB != VK_NULL_HANDLE);
+				}
 				LastAliasedRenderTarget = RenderTargetTexture;
 				UE_LOG(LogMagicLeap, Log, TEXT("Aliased render target for correct sRGB ouput."));
 			}
@@ -588,7 +603,9 @@ void FMagicLeapCustomPresentVulkan::FinishRendering()
 		MLResult Result = MLGraphicsEndFrame(Plugin->GraphicsClient, Plugin->GetCurrentFrame().Handle);
 		if (Result != MLResult_Ok)
 		{
+#if !WITH_EDITOR
 			UE_LOG(LogMagicLeap, Error, TEXT("MLGraphicsEndFrame failed with status %d"), Result);
+#endif
 		}
 	}
   
@@ -619,7 +636,7 @@ void FMagicLeapCustomPresentVulkan::UpdateViewport(const FViewport& Viewport, FR
 	check(IsValidRef(RT));
 
 	RenderTargetTexture = RT->GetNativeResource();
-	RenderTargetTextureAllocation = static_cast<FVulkanTexture2D*>(RT->GetTexture2D())->Surface.GetAllocationHandle();
+	RenderTargetTextureAllocation = reinterpret_cast<void*>(static_cast<FVulkanTexture2D*>(RT->GetTexture2D())->Surface.GetAllocationHandle());
 	RenderTargetTextureAllocationOffset = static_cast<FVulkanTexture2D*>(RT->GetTexture2D())->Surface.GetAllocationOffset();
 
 	InViewportRHI->SetCustomPresent(this);
diff --git a/Engine/Plugins/Lumin/MagicLeap/Source/MagicLeap/Private/MagicLeapHMD.cpp b/Engine/Plugins/Lumin/MagicLeap/Source/MagicLeap/Private/MagicLeapHMD.cpp
--- a/Engine/Plugins/Lumin/MagicLeap/Source/MagicLeap/Private/MagicLeapHMD.cpp
+++ b/Engine/Plugins/Lumin/MagicLeap/Source/MagicLeap/Private/MagicLeapHMD.cpp
@@ -17,6 +17,8 @@
 #include "DeviceProfiles/DeviceProfileManager.h"
 #include "DeviceProfiles/DeviceProfile.h"
 #include "LuminRuntimeSettings.h"
+#include "MotionControllerComponent.h"
+#include "IMagicLeapInputDevice.h"
 
 #include "XRThreadUtils.h"
 
@@ -33,6 +35,7 @@
 #include "GeneralProjectSettings.h"
 #include "MagicLeapSettings.h"
 #include "MagicLeapSDKDetection.h"
+#include "IMagicLeapModule.h"
 
 #if !PLATFORM_MAC // @todo Lumin: I had to add this to get Mac to compile - trying to add GL to Mac build had massive compile issues
 #include "OpenGLDrv.h"
@@ -123,7 +126,7 @@ public:
 				IConsoleVariable* CVar = IConsoleManager::Get().FindConsoleVariable(TEXT("r.Vulkan.RHIThread"));
 				if (CVar)
 				{
-					CVar->Set(TEXT("0"));
+					// CVar->Set(TEXT("0"));
 				}
 			}
 			else
@@ -162,7 +165,20 @@ public:
 	/** IHeadMountedDisplayModule implementation */
 	virtual TSharedPtr< class IXRTrackingSystem, ESPMode::ThreadSafe > CreateTrackingSystem() override
 	{
-		auto LocalHMD = GetOrCreateHMD();
+		TSharedPtr<FMagicLeapHMD, ESPMode::ThreadSafe> LocalHMD;
+#if !PLATFORM_MAC
+		if (HMD.IsValid())
+		{
+			LocalHMD = HMD.Pin();
+		}
+		else
+		{
+			//initialize AR system
+			TSharedPtr<IARSystemSupport, ESPMode::ThreadSafe> ARImplementation;
+			LocalHMD = MakeShared<FMagicLeapHMD, ESPMode::ThreadSafe>(this, ARImplementation.Get(), bIsVDZIEnabled, bUseVulkanForZI);
+			HMD = LocalHMD;
+		}
+#endif
 		if (LocalHMD.IsValid() && !LocalHMD->IsInitialized())
 		{
 			LocalHMD->Startup();
@@ -192,7 +208,18 @@ public:
 
 	virtual TWeakPtr<IMagicLeapHMD, ESPMode::ThreadSafe> GetHMD() override
 	{
-		return GetOrCreateHMD();
+#if WITH_EDITOR
+		if (bIsVDZIEnabled)
+		{
+			return HMD;
+		}
+		else
+		{
+			return nullptr;
+		}
+#else
+		return HMD;
+#endif
 	}
 
 	virtual TSharedPtr<IHeadMountedDisplayVulkanExtensions, ESPMode::ThreadSafe> GetVulkanExtensions() override
@@ -207,6 +234,45 @@ public:
 		return nullptr;
 	}
 
+	virtual void RegisterMagicLeapInputDevice(IMagicLeapInputDevice* InputDevice) override
+	{
+		check(InputDevice);
+		InputDevices.Add(InputDevice);
+	}
+
+	virtual void UnregisterMagicLeapInputDevice(IMagicLeapInputDevice* InputDevice) override
+	{
+		check(InputDevice);
+		InputDevices.Remove(InputDevice);
+	}
+
+	virtual void EnableInputDevices() override
+	{
+		for (auto& It : InputDevices)
+		{
+			if (It->SupportsExplicitEnable())
+			{
+				It->Enable();
+			}
+		}
+	}
+
+	virtual void DisableInputDevices() override
+	{
+		for (auto& It : InputDevices)
+		{
+			It->Disable();
+		}
+	}
+
+	virtual void OnBeginRendering_GameThread_UpdateInputDevices() override
+	{
+		for (auto& It : InputDevices)
+		{
+			It->OnBeginRendering_GameThread_Update();
+		}
+	}
+
 private:
 
 #if WITH_EDITOR
@@ -237,36 +303,12 @@ private:
 	}
 #endif
 
-	TSharedPtr<FMagicLeapHMD, ESPMode::ThreadSafe> GetOrCreateHMD()
-	{
-#if !PLATFORM_MAC
-		if (!HMD.IsValid())
-		{
-			//initialize AR system
-			TSharedPtr<IARSystemSupport, ESPMode::ThreadSafe> ARImplementation;
-
-			HMD = MakeShared<FMagicLeapHMD, ESPMode::ThreadSafe>(this, ARImplementation.Get(), bIsVDZIEnabled, bUseVulkanForZI);
-		}
-#endif
-#if WITH_EDITOR
-		if (bIsVDZIEnabled)
-		{
-			return HMD;
-		}
-		else
-		{
-			return nullptr;
-		}
-#else
-		return HMD;
-#endif
-	}
-
 	bool bIsVDZIEnabled;
 	bool bUseVulkanForZI;
 	FMagicLeapAPISetup APISetup;
-	TSharedPtr<FMagicLeapHMD, ESPMode::ThreadSafe> HMD;
+	TWeakPtr<FMagicLeapHMD, ESPMode::ThreadSafe> HMD;
 	TSharedPtr<FMagicLeapVulkanExtensions, ESPMode::ThreadSafe> VulkanExtensions;
+	TSet<IMagicLeapInputDevice*> InputDevices;
 };
 
 IMPLEMENT_MODULE(FMagicLeapPlugin, MagicLeap)
@@ -365,7 +407,7 @@ bool FMagicLeapHMD::IsHMDConnected()
 			MLResult Result = MLRemoteIsServerConfigured(&bZIServerRunning);
 			if (Result != MLResult_Ok)
 			{
-				UE_LOG(LogMagicLeap, Error, TEXT("MLRemoteIsServerConfigured failed with error %d!"), Result);
+				UE_LOG(LogMagicLeap, Error, TEXT("MLRemoteIsServerConfigured failed with error %s!"), UTF8_TO_TCHAR(MLGetResultString(Result)));
 
 				// Ensure we don't falsely mark it as running if there was an error
 				bZIServerRunning = false;
@@ -587,13 +629,13 @@ void FMagicLeapHMD::RefreshTrackingFrame()
 	// are ticked and fire their events before the OnStartGameFrame().
 	MLResult Result = MLPerceptionReleaseSnapshot(GameTrackingFrame.Snapshot);
 #if PLATFORM_LUMIN
-	UE_CLOG(Result != MLResult_Ok, LogMagicLeap, Error, TEXT("MLImageTrackerInitSettings failed with error %d."), Result);
+	UE_CLOG(Result != MLResult_Ok, LogMagicLeap, Error, TEXT("MLPerceptionReleaseSnapshot failed with error %s."), UTF8_TO_TCHAR(MLGetResultString(Result)));
 #endif //PLATFORM_LUMIN
 
 	// get the snapshot for the frame
 	Result = MLPerceptionGetSnapshot(&GameTrackingFrame.Snapshot);
 #if PLATFORM_LUMIN
-	UE_CLOG(Result != MLResult_Ok, LogMagicLeap, Error, TEXT("MLPerceptionGetSnapshot failed with error %d."), Result);
+	UE_CLOG(Result != MLResult_Ok, LogMagicLeap, Error, TEXT("MLPerceptionGetSnapshot failed with error %s."), UTF8_TO_TCHAR(MLGetResultString(Result)));
 #endif //PLATFORM_LUMIN
 
 	MLHeadTrackingState state;
@@ -1015,6 +1057,7 @@ FMagicLeapHMD::FMagicLeapHMD(IMagicLeapPlugin* InMagicLeapPlugin, IARSystemSuppo
 	bVDZIWarningDisplayed(false),
 	bPrivilegesEnabled(false),
 	CurrentFrameTimingHint(ELuminFrameTimingHint::Unspecified),
+	bQueuedGraphicsCreateCall(false),
 	bHeadTrackingStateAvailable(false)
 {
 #if WITH_EDITOR
@@ -1210,6 +1253,7 @@ void FMagicLeapHMD::DisableDeviceFeatures()
 	RestoreBaseProfile();
 	DisableInputDevices();
 	DisableHeadTracking();
+	DISABLE_MAGIC_LEAP_MODULE("MagicLeapEyeTracker");
 	DisablePerception();
 	DisablePrivileges();
 	if (GIsEditor)
@@ -1229,6 +1273,12 @@ ML_EXTERN_C_END
 void FMagicLeapHMD::InitDevice_RenderThread()
 {
 #if WITH_MLSDK
+	if (bQueuedGraphicsCreateCall)
+	{
+		UE_LOG(LogMagicLeap, Warning, TEXT("Graphics client create call already queued."));
+		return;
+	}
+
 	if (!bDeviceInitialized)
 	{
 		bool bDeviceSuccessfullyInitialized = false;
@@ -1238,13 +1288,7 @@ void FMagicLeapHMD::InitDevice_RenderThread()
 		gfx_opts.color_format = MLSurfaceFormat_RGBA8UNormSRGB;
 		gfx_opts.depth_format = MLSurfaceFormat_D32Float;
 
-		// HACK - Forward rendering on Lumin assumes a bottom-left origin.
-		// There might be other platforms that need this that we haven't accounted for.
-#if !PLATFORM_LUMIN || PLATFORM_LUMINGL4
-		gfx_opts.graphics_flags = MLGraphicsFlags_OriginUpperLeft;
-#else
 		gfx_opts.graphics_flags = MLGraphicsFlags_Default;
-#endif
 
 #if PLATFORM_WINDOWS
 		if (IsPCPlatform(GMaxRHIShaderPlatform) && !IsOpenGLPlatform(GMaxRHIShaderPlatform))
@@ -1279,7 +1323,7 @@ void FMagicLeapHMD::InitDevice_RenderThread()
 			{
 				bDeviceSuccessfullyInitialized = false;
 				GraphicsClient = ML_INVALID_HANDLE;
-				UE_LOG(LogMagicLeap, Error, TEXT("MLGraphicsCreateClientGL failed with status %d"), Result);
+				UE_LOG(LogMagicLeap, Error, TEXT("MLGraphicsCreateClientGL failed with status %s"), UTF8_TO_TCHAR(MLGetResultString(Result)));
 			}
 		}
 		if (IsVulkanPlatform(GMaxRHIShaderPlatform))
@@ -1288,7 +1332,8 @@ void FMagicLeapHMD::InitDevice_RenderThread()
 			static const auto* VulkanRHIThread = IConsoleManager::Get().FindTConsoleVariableDataInt(TEXT("r.Vulkan.RHIThread"));
 			UE_LOG(LogMagicLeap, Warning, TEXT("r.Vulkan.RHIThread=%d"), VulkanRHIThread->GetValueOnAnyThread());
 
-			gfx_opts.graphics_flags = MLGraphicsFlags_OriginUpperLeft;
+			bQueuedGraphicsCreateCall = true;
+
 			ExecuteOnRHIThread_DoNotWait([this, gfx_opts]()
 			{
 				UE_LOG(LogMagicLeap, Display, TEXT("FMagicLeapCustomPresentVulkan is supported."));
@@ -1311,11 +1356,12 @@ void FMagicLeapHMD::InitDevice_RenderThread()
 				else
 				{
 					GraphicsClient = ML_INVALID_HANDLE;
-					UE_LOG(LogMagicLeap, Error, TEXT("MLGraphicsCreateClientVk failed with status %d"), Result);
+					UE_LOG(LogMagicLeap, Error, TEXT("MLGraphicsCreateClientVk failed with status %s"), UTF8_TO_TCHAR(MLGetResultString(Result)));
 				}
 
 				FPlatformAtomics::InterlockedExchange(&bDeviceInitialized, Result == MLResult_Ok);
 				FPlatformAtomics::InterlockedExchange(&bDeviceWasJustInitialized, Result == MLResult_Ok);
+				bQueuedGraphicsCreateCall = false;
 			});
 #endif // PLATFORM_WINDOWS || PLATFORM_LUMIN
 		}
@@ -1473,7 +1519,7 @@ void FMagicLeapHMD::ReleaseDevice_RenderThread()
 		MLResult Result = MLGraphicsDestroyClient(&GraphicsClient);
 		if (Result != MLResult_Ok)
 		{
-			UE_LOG(LogMagicLeap, Error, TEXT("MLGraphicsDestroyClient failed with status %d"), Result);
+			UE_LOG(LogMagicLeap, Error, TEXT("MLGraphicsDestroyClient failed with status %s"), UTF8_TO_TCHAR(MLGetResultString(Result)));
 		}
 		else
 		{
@@ -1556,7 +1602,8 @@ void FMagicLeapHMD::OnBeginRendering_GameThread()
 	ExecuteOnRenderThread_DoNotWait([this, TrackingFrameCopy = GameTrackingFrame]() 
 	{
 		MLSnapshot* OldSnapshot = RenderTrackingFrame.Snapshot;
-		RenderTrackingFrame = TrackingFrameCopy;
+		// Don't update RenderTrackingFrame here. It is updated by the RHITrackingFrame in FMagicLeapCustomPreset::BeginRendering()
+		// RenderTrackingFrame = TrackingFrameCopy;
 #if !PLATFORM_MAC
 		ExecuteOnRHIThread_DoNotWait([this, TrackingFrameCopy]()
 		{
@@ -1565,6 +1612,9 @@ void FMagicLeapHMD::OnBeginRendering_GameThread()
 #endif //PLATFORM_MAC
 	});
 #endif //WITH_MLSDK
+
+	// Update the devices, in particular input controller devices.
+	IMagicLeapPlugin::Get().OnBeginRendering_GameThread_UpdateInputDevices();
 }
 
 TSharedPtr<class IXRCamera, ESPMode::ThreadSafe> FMagicLeapHMD::GetXRCamera(int32 DeviceId /*= HMDDeviceId*/)
@@ -1688,18 +1738,6 @@ void FMagicLeapHMD::SetClippingPlanes(float NCP, float FCP)
 	UpdateNearClippingPlane();
 }
 
-void FMagicLeapHMD::RegisterMagicLeapInputDevice(IMagicLeapInputDevice* InputDevice)
-{
-	check(InputDevice);
-	InputDevices.Add(InputDevice);
-}
-
-void FMagicLeapHMD::UnregisterMagicLeapInputDevice(IMagicLeapInputDevice* InputDevice)
-{
-	check(InputDevice);
-	InputDevices.Remove(InputDevice);
-}
-
 bool FMagicLeapHMD::IsInitialized() const
 {
 	return (AppFramework.IsInitialized());
@@ -1792,6 +1830,11 @@ void FMagicLeapHMD::InitializeOldFrameFromRenderFrame()
 	}
 }
 
+void FMagicLeapHMD::InitializeRenderFrameFromRHIFrame()
+{
+	RenderTrackingFrame = RHITrackingFrame;
+}
+
 const FAppFramework& FMagicLeapHMD::GetAppFrameworkConst() const
 {
 	return AppFramework;
@@ -1928,21 +1971,12 @@ void FMagicLeapHMD::DisablePrivileges()
 
 void FMagicLeapHMD::EnableInputDevices()
 {
-	for (auto& It : InputDevices)
-	{
-		if (It->SupportsExplicitEnable())
-		{
-			It->Enable();
-		}
-	}
+	IMagicLeapPlugin::Get().EnableInputDevices();
 }
 
 void FMagicLeapHMD::DisableInputDevices()
 {
-	for (auto& It : InputDevices)
-	{
-		It->Disable();
-	}
+	IMagicLeapPlugin::Get().DisableInputDevices();
 }
 
 void FMagicLeapHMD::EnablePerception()
@@ -1961,12 +1995,12 @@ void FMagicLeapHMD::EnablePerception()
 			}
 			else
 			{
-				UE_LOG(LogMagicLeap, Error, TEXT("MLPerceptionStartup failed with error %d."), Result);
+				UE_LOG(LogMagicLeap, Error, TEXT("MLPerceptionStartup failed with error %s."), UTF8_TO_TCHAR(MLGetResultString(Result)));
 			}
 		}
 		else
 		{
-			UE_LOG(LogMagicLeap, Error, TEXT("MLPerceptionInitSettings failed with error %d."), Result);
+			UE_LOG(LogMagicLeap, Error, TEXT("MLPerceptionInitSettings failed with error %s."), UTF8_TO_TCHAR(MLGetResultString(Result)));
 		}
 	}
 #endif //WITH_MLSDK
@@ -1985,7 +2019,7 @@ void FMagicLeapHMD::DisablePerception()
 		}
 		else
 		{
-			UE_LOG(LogMagicLeap, Error, TEXT("MLPerceptionShutdown failed with error %d"), Result);
+			UE_LOG(LogMagicLeap, Error, TEXT("MLPerceptionShutdown failed with error %s"), UTF8_TO_TCHAR(MLGetResultString(Result)));
 		}
 	}
 #endif //WITH_MLSDK
@@ -1996,17 +2030,17 @@ void FMagicLeapHMD::EnableHeadTracking()
 #if WITH_MLSDK
 	if (HeadTracker == ML_INVALID_HANDLE)
 	{
-		MLResult CreateResult = MLHeadTrackingCreate(&HeadTracker);
-		if (CreateResult == MLResult_Ok && HeadTracker != ML_INVALID_HANDLE)
+		MLResult Result = MLHeadTrackingCreate(&HeadTracker);
+		if (Result == MLResult_Ok && HeadTracker != ML_INVALID_HANDLE)
 		{
 			if (MLResult_Ok != MLHeadTrackingGetStaticData(HeadTracker, &HeadTrackerData))
 			{
-				UE_LOG(LogMagicLeap, Error, TEXT("Error getting head tracking static data."));
+				UE_LOG(LogMagicLeap, Error, TEXT("MLHeadTrackingGetStaticData failed with error %s."), UTF8_TO_TCHAR(MLGetResultString(Result)));
 			}
 		}
 		else
 		{
-			UE_LOG(LogMagicLeap, Error, TEXT("Error starting head tracking."));
+			UE_LOG(LogMagicLeap, Error, TEXT("MLHeadTrackingCreate failed with error %s."), UTF8_TO_TCHAR(MLGetResultString(Result)));
 		}
 	}
 #endif //WITH_MLSDK
@@ -2017,8 +2051,8 @@ void FMagicLeapHMD::DisableHeadTracking()
 #if WITH_MLSDK
 	if (HeadTracker != ML_INVALID_HANDLE)
 	{
-		MLResult DestroyResult = MLHeadTrackingDestroy(HeadTracker);
-		UE_CLOG(DestroyResult != MLResult_Ok, LogMagicLeap, Error, TEXT("Error destroying head tracking."));
+		MLResult Result = MLHeadTrackingDestroy(HeadTracker);
+		UE_CLOG(Result != MLResult_Ok, LogMagicLeap, Error, TEXT("MLHeadTrackingDestroy failed with error %s."), UTF8_TO_TCHAR(MLGetResultString(Result)));
 		HeadTracker = ML_INVALID_HANDLE;
 	}
 #endif //WITH_MLSDK
@@ -2037,14 +2071,14 @@ void FMagicLeapHMD::InitializeClipExtents_RenderThread()
 	}
 	else
 	{
-		UE_LOG(LogMagicLeap, Error, TEXT("MLGraphicsGetRenderTargets() failed with status %d"), Result);
+		UE_LOG(LogMagicLeap, Error, TEXT("MLGraphicsGetRenderTargets failed with error %s"), UTF8_TO_TCHAR(MLGetResultString(Result)));
 	}
 
 	// get the clip extents for clipping content in update thread
 	Result = MLGraphicsGetClipExtents(GraphicsClient, &GameTrackingFrame.UpdateInfoArray);
 	if (Result != MLResult_Ok)
 	{
-		FString ErrorMesg = FString::Printf(TEXT("MLGraphicsGetClipExtents failed with status %d"), Result);
+		FString ErrorMesg = FString::Printf(TEXT("MLGraphicsGetClipExtents failed with error %s"), UTF8_TO_TCHAR(MLGetResultString(Result)));
 
 		// In case we're running under VD/ZI, there's always the risk of disconnects.
 		// In those cases, the graphics API can return an error, but the client handle might still be valid.
@@ -2105,18 +2139,4 @@ FMagicLeapHMD* FMagicLeapHMD::GetHMD()
 }
 #endif
 
-bool FMagicLeapHMD::FMagicLeapVulkanExtensions::GetVulkanInstanceExtensionsRequired(TArray<const ANSICHAR*>& Out)
-{
-	return true;
-}
-
-bool FMagicLeapHMD::FMagicLeapVulkanExtensions::GetVulkanDeviceExtensionsRequired(struct VkPhysicalDevice_T *pPhysicalDevice, TArray<const ANSICHAR*>& Out)
-{
-#if !PLATFORM_MAC
-	return FMagicLeapHelperVulkan::GetVulkanDeviceExtensionsRequired(pPhysicalDevice, Out);
-#else
-	return false;
-#endif
-}
-
 #undef LOCTEXT_NAMESPACE
diff --git a/Engine/Plugins/Lumin/MagicLeap/Source/MagicLeap/Private/MagicLeapHMD.h b/Engine/Plugins/Lumin/MagicLeap/Source/MagicLeap/Private/MagicLeapHMD.h
--- a/Engine/Plugins/Lumin/MagicLeap/Source/MagicLeap/Private/MagicLeapHMD.h
+++ b/Engine/Plugins/Lumin/MagicLeap/Source/MagicLeap/Private/MagicLeapHMD.h
@@ -7,6 +7,7 @@
 #include "IMagicLeapPlugin.h"
 #include "IMagicLeapHMD.h"
 #include "Misc/ScopeLock.h"
+#include "HAL/ThreadSafeBool.h"
 
 #include "XRRenderTargetManager.h"
 #include "IStereoLayers.h"
@@ -17,7 +18,6 @@
 #include "MagicLeapHMDFunctionLibrary.h"
 #include "LuminRuntimeSettings.h"
 #include "MagicLeapPluginUtil.h" // for ML_INCLUDES_START/END
-#include "IHeadMountedDisplayVulkanExtensions.h"
 
 #if WITH_MLSDK
 ML_INCLUDES_START
@@ -112,19 +112,6 @@ public:
 		return false;
 	}
 
-	/** Vulkan Extensions */
-	class FMagicLeapVulkanExtensions : public IHeadMountedDisplayVulkanExtensions, public TSharedFromThis<FMagicLeapVulkanExtensions, ESPMode::ThreadSafe>
-	{
-	public:
-		FMagicLeapVulkanExtensions() {}
-		virtual ~FMagicLeapVulkanExtensions() {}
-
-		// IHeadMountedDisplayVulkanExtensions
-		virtual bool GetVulkanInstanceExtensionsRequired(TArray<const ANSICHAR*>& Out) override;
-		virtual bool GetVulkanDeviceExtensionsRequired(struct VkPhysicalDevice_T *pPhysicalDevice, TArray<const ANSICHAR*>& Out) override;
-	};
-
-
 	IStereoLayers* GetStereoLayers() override;
 
 	// FXRRenderTargetManager interface
@@ -143,8 +130,6 @@ public:
 	virtual ~FMagicLeapHMD();
 
 	/** FMagicLeapHMDBase interface */
-	virtual void RegisterMagicLeapInputDevice(IMagicLeapInputDevice* InputDevice) override;
-	virtual void UnregisterMagicLeapInputDevice(IMagicLeapInputDevice* InputDevice) override;
 	virtual bool IsInitialized() const override;
 	bool IsDeviceInitialized() const { return (bDeviceInitialized != 0) ? true : false; }
 
@@ -200,6 +185,7 @@ public:
 	// HACK: This is a hack in order to use projection matrices from last render frame
 	// This should be removed once unreal can use separate projection matrices for update and render
 	void InitializeOldFrameFromRenderFrame();
+	void InitializeRenderFrameFromRHIFrame();
 
 	const FAppFramework& GetAppFrameworkConst() const;
 	FAppFramework& GetAppFramework();
@@ -214,11 +200,6 @@ public:
 	MLHandle GraphicsClient;
 #endif //WITH_MLSDK
 
-private:
-
-	// The FMagicLeapPlugin class functions as a factory for FMagicLeapHMD.
-	// Therefore, it needs some control over initalization & termination.
-	friend class FMagicLeapPlugin;
 
 	/**
 	* Starts up the SensoryWare API
@@ -230,6 +211,7 @@ private:
 	*/
 	void Shutdown();
 
+private:
 	void EnableDeviceFeatures();
 	void DisableDeviceFeatures();
 
@@ -326,8 +308,8 @@ public:
 	FTrackingFrame OldTrackingFrame;
 
 private:
-	TSet<IMagicLeapInputDevice*> InputDevices;
 	TWeakObjectPtr<const AActor> FocusActor;
+	FThreadSafeBool bQueuedGraphicsCreateCall;
 
 	struct SavedProfileState
 	{
diff --git a/Engine/Plugins/Lumin/MagicLeap/Source/MagicLeap/Private/MagicLeapPool.h b/Engine/Plugins/Lumin/MagicLeap/Source/MagicLeap/Private/MagicLeapPool.h
new file mode 100644
--- /dev/null
+++ b/Engine/Plugins/Lumin/MagicLeap/Source/MagicLeap/Private/MagicLeapPool.h
@@ -0,0 +1,101 @@
+// %BANNER_BEGIN%
+// ---------------------------------------------------------------------
+// %COPYRIGHT_BEGIN%
+//
+// Copyright (c) 2017 Magic Leap, Inc. (COMPANY) All Rights Reserved.
+// Magic Leap, Inc. Confidential and Proprietary
+//
+// NOTICE: All information contained herein is, and remains the property
+// of COMPANY. The intellectual and technical concepts contained herein
+// are proprietary to COMPANY and may be covered by U.S. and Foreign
+// Patents, patents in process, and are protected by trade secret or
+// copyright law. Dissemination of this information or reproduction of
+// this material is strictly forbidden unless prior written permission is
+// obtained from COMPANY. Access to the source code contained herein is
+// hereby forbidden to anyone except current COMPANY employees, managers
+// or contractors who have executed Confidentiality and Non-disclosure
+// agreements explicitly covering such access.
+//
+// The copyright notice above does not evidence any actual or intended
+// publication or disclosure of this source code, which includes
+// information that is confidential and/or proprietary, and is a trade
+// secret, of COMPANY. ANY REPRODUCTION, MODIFICATION, DISTRIBUTION,
+// PUBLIC PERFORMANCE, OR PUBLIC DISPLAY OF OR THROUGH USE OF THIS
+// SOURCE CODE WITHOUT THE EXPRESS WRITTEN CONSENT OF COMPANY IS
+// STRICTLY PROHIBITED, AND IN VIOLATION OF APPLICABLE LAWS AND
+// INTERNATIONAL TREATIES. THE RECEIPT OR POSSESSION OF THIS SOURCE
+// CODE AND/OR RELATED INFORMATION DOES NOT CONVEY OR IMPLY ANY RIGHTS
+// TO REPRODUCE, DISCLOSE OR DISTRIBUTE ITS CONTENTS, OR TO MANUFACTURE,
+// USE, OR SELL ANYTHING THAT IT MAY DESCRIBE, IN WHOLE OR IN PART.
+//
+// %COPYRIGHT_END%
+// --------------------------------------------------------------------
+// %BANNER_END%
+
+#pragma once
+
+#include "Containers/Array.h"
+
+#define DEBUG_MEM_POOL 0
+
+#if DEBUG_MEM_POOL
+DEFINE_LOG_CATEGORY_STATIC(MagicLeapMemPool, Log, All);
+#endif // DEBUG_MEM_POOL
+
+template <class T>
+class FMagicLeapPool
+{
+public:
+	FMagicLeapPool(int32 InPoolSize)
+	: Size(InPoolSize)
+	{
+		Blocks.AddZeroed(1);
+		Blocks[0].AddZeroed(Size);
+		for (int32 i = 0; i < Size; ++i)
+		{
+			Free.Push(&Blocks[0][i]);
+		}
+	}
+
+	T* GetNextFree()
+	{
+		if (Free.Num() == 0)
+		{
+#if DEBUG_MEM_POOL
+			UE_LOG(MagicLeapMemPool, Log, TEXT("FMagicLeapPool is out of space.  Allocating new block."));
+#endif // DEBUG_MEM_POOL
+			int32 NewBlockIndex = Blocks.AddZeroed(1);
+			Blocks[NewBlockIndex].AddZeroed(Size);
+			for (int32 i = 0; i < Size; ++i)
+			{
+				Free.Push(&Blocks[NewBlockIndex][i]);
+			}
+		}
+#if DEBUG_MEM_POOL
+		T* NewAllocation = Free.Pop();
+		UE_LOG(MagicLeapMemPool, Log, TEXT("FMagicLeapPool allocated %p."), NewAllocation);
+		Allocated.Push(NewAllocation);
+		return NewAllocation;
+#else
+		return Free.Pop();
+#endif // DEBUG_MEM_POOL
+	}
+
+	void Release(T* InAllocation)
+	{
+#if DEBUG_MEM_POOL
+		checkf(Allocated.Contains(InAllocation), TEXT("Attempting to release memory that was not allocated by this pool!"));
+		Allocated.Pop(InAllocation);
+		UE_LOG(MagicLeapMemPool, Log, TEXT("FMagicLeapPool released %p."), InAllocation);
+#endif // DEBUG_MEM_POOL
+		Free.Push(InAllocation);
+	}
+
+private:
+	const int32 Size;
+	TArray<TArray<T>> Blocks;
+	TArray<T*> Free;
+#if DEBUG_MEM_POOL
+	TArray<T*> Allocated;
+#endif // DEBUG_MEM_POOL
+};
\ No newline at end of file
diff --git a/Engine/Plugins/Lumin/MagicLeap/Source/MagicLeap/Private/MagicLeapRunnable.h b/Engine/Plugins/Lumin/MagicLeap/Source/MagicLeap/Private/MagicLeapRunnable.h
new file mode 100644
--- /dev/null
+++ b/Engine/Plugins/Lumin/MagicLeap/Source/MagicLeap/Private/MagicLeapRunnable.h
@@ -0,0 +1,223 @@
+// %BANNER_BEGIN%
+// ---------------------------------------------------------------------
+// %COPYRIGHT_BEGIN%
+//
+// Copyright (c) 2017 Magic Leap, Inc. (COMPANY) All Rights Reserved.
+// Magic Leap, Inc. Confidential and Proprietary
+//
+// NOTICE: All information contained herein is, and remains the property
+// of COMPANY. The intellectual and technical concepts contained herein
+// are proprietary to COMPANY and may be covered by U.S. and Foreign
+// Patents, patents in process, and are protected by trade secret or
+// copyright law. Dissemination of this information or reproduction of
+// this material is strictly forbidden unless prior written permission is
+// obtained from COMPANY. Access to the source code contained herein is
+// hereby forbidden to anyone except current COMPANY employees, managers
+// or contractors who have executed Confidentiality and Non-disclosure
+// agreements explicitly covering such access.
+//
+// The copyright notice above does not evidence any actual or intended
+// publication or disclosure of this source code, which includes
+// information that is confidential and/or proprietary, and is a trade
+// secret, of COMPANY. ANY REPRODUCTION, MODIFICATION, DISTRIBUTION,
+// PUBLIC PERFORMANCE, OR PUBLIC DISPLAY OF OR THROUGH USE OF THIS
+// SOURCE CODE WITHOUT THE EXPRESS WRITTEN CONSENT OF COMPANY IS
+// STRICTLY PROHIBITED, AND IN VIOLATION OF APPLICABLE LAWS AND
+// INTERNATIONAL TREATIES. THE RECEIPT OR POSSESSION OF THIS SOURCE
+// CODE AND/OR RELATED INFORMATION DOES NOT CONVEY OR IMPLY ANY RIGHTS
+// TO REPRODUCE, DISCLOSE OR DISTRIBUTE ITS CONTENTS, OR TO MANUFACTURE,
+// USE, OR SELL ANYTHING THAT IT MAY DESCRIBE, IN WHOLE OR IN PART.
+//
+// %COPYRIGHT_END%
+// --------------------------------------------------------------------
+// %BANNER_END%
+
+#pragma once
+
+#include "HAL/Runnable.h"
+#include "HAL/RunnableThread.h"
+#include "HAL/ThreadSafeBool.h"
+#include "HAL/ThreadSafeCounter64.h"
+#include "Containers/Queue.h"
+#include "Containers/Array.h"
+#include "Misc/Variant.h"
+#include "AppEventHandler.h"
+#include "MagicLeapPluginUtil.h" // for ML_INCLUDES_START/END
+#if WITH_MLSDK
+ML_INCLUDES_START
+#include <ml_api.h>
+ML_INCLUDES_END
+#endif //WITH_MLSDK
+
+struct FMagicLeapTask
+{
+	bool bSuccess;
+
+	FMagicLeapTask()
+	: bSuccess(false)
+	{
+	}
+};
+
+template <class TTaskType>
+class FMagicLeapRunnable : public FRunnable
+{
+public:
+#if WITH_MLSDK
+	FMagicLeapRunnable(const TArray<MLPrivilegeID>& InRequiredPrivileges, const FString& InName)
+	: AppEventHandler(InRequiredPrivileges)
+#else
+	FMagicLeapRunnable(const FString& InName)
+	: AppEventHandler()
+#endif //WITH_MLSDK
+	, Thread(nullptr)
+	, StopTaskCounter(0)
+	, Semaphore(nullptr)
+	, bPaused(false)
+	{
+		Semaphore = FGenericPlatformProcess::GetSynchEventFromPool(false);
+#if PLATFORM_LUMIN
+		Thread = FRunnableThread::Create(this, *InName, 0, TPri_BelowNormal, FLuminAffinity::GetPoolThreadMask());
+#else
+		Thread = FRunnableThread::Create(this, *InName, 0, TPri_BelowNormal);
+#endif // PLATFORM_LUMIN
+
+		AppEventHandler.SetOnAppPauseHandler([this]() 
+		{
+			OnAppPause();
+		});
+
+		AppEventHandler.SetOnAppResumeHandler([this]()
+		{
+			OnAppResume();
+		});
+
+		AppEventHandler.SetOnAppShutDownHandler([this]()
+		{
+			OnAppShutDown();
+		});
+	}
+
+	virtual ~FMagicLeapRunnable()
+	{
+		StopTaskCounter.Increment();
+
+		if (Semaphore)
+		{
+			Semaphore->Trigger();
+			Thread->WaitForCompletion();
+			FGenericPlatformProcess::ReturnSynchEventToPool(Semaphore);
+			Semaphore = nullptr;
+		}
+
+		delete Thread;
+		Thread = nullptr;
+	}
+
+	uint32 Run() override
+	{
+		while (StopTaskCounter.GetValue() == 0)
+		{
+			if (bPaused)
+			{
+				Pause();
+				// Cancel any incoming tasks.
+				CancelIncomingTasks();
+				// Wait for signal from resume call.
+				Semaphore->Wait();
+				Resume();
+			}
+			else if (!IncomingTasks.IsEmpty())
+			{
+				DoNextTask();
+			}
+			else
+			{
+				Semaphore->Wait();
+			}
+		}
+
+		return 0;
+	}
+
+	void OnAppPause()
+	{
+		bPaused = true;
+		Semaphore->Trigger();
+	}
+
+	void OnAppResume()
+	{
+		bPaused = false;
+		Semaphore->Trigger();
+	}
+
+	void OnAppShutDown()
+	{
+		Stop();
+	}
+
+	void PushNewTask(TTaskType InTask)
+	{
+		IncomingTasks.Enqueue(InTask);
+		// wake up the worker to process the task
+		Semaphore->Trigger();
+	}
+
+	void PushCompletedTask(TTaskType InTask)
+	{
+		CompletedTasks.Enqueue(InTask);
+	}
+
+	bool TryGetCompletedTask(TTaskType& OutCompletedTask)
+	{
+#if PLATFORM_LUMIN
+		if (CompletedTasks.Peek(OutCompletedTask))
+		{
+			CompletedTasks.Pop();
+			return true;
+		}
+#endif //PLATFORM_LUMIN
+		return false;
+	}
+
+protected:
+	void CancelIncomingTasks()
+	{
+		while (IncomingTasks.Dequeue(CurrentTask))
+		{
+			CurrentTask.bSuccess = false;
+			CompletedTasks.Enqueue(CurrentTask);
+		}
+	}
+
+	virtual bool ProcessCurrentTask() = 0;
+	virtual void Pause() {}
+	virtual void Resume() {}
+
+	/** Internal thread this instance is running on */
+	FRunnableThread* Thread;
+	FThreadSafeCounter StopTaskCounter;
+	FEvent* Semaphore;
+	FThreadSafeBool bPaused;
+	TQueue<TTaskType, EQueueMode::Spsc> IncomingTasks;
+	TQueue<TTaskType, EQueueMode::Spsc> CompletedTasks;
+	TTaskType CurrentTask;
+	MagicLeap::IAppEventHandler AppEventHandler;
+
+private:
+	bool DoNextTask()
+	{
+		IncomingTasks.Dequeue(CurrentTask);
+		CurrentTask.bSuccess = ProcessCurrentTask();
+
+		if (bPaused)
+		{
+			return false;
+		}
+
+		CompletedTasks.Enqueue(CurrentTask);
+
+		return CurrentTask.bSuccess;
+	}
+};
\ No newline at end of file
diff --git a/Engine/Plugins/Lumin/MagicLeap/Source/MagicLeap/Private/MagicLeapStereoLayers.cpp b/Engine/Plugins/Lumin/MagicLeap/Source/MagicLeap/Private/MagicLeapStereoLayers.cpp
--- a/Engine/Plugins/Lumin/MagicLeap/Source/MagicLeap/Private/MagicLeapStereoLayers.cpp
+++ b/Engine/Plugins/Lumin/MagicLeap/Source/MagicLeap/Private/MagicLeapStereoLayers.cpp
@@ -11,25 +11,48 @@
 
 IStereoLayers* FMagicLeapHMD::GetStereoLayers()
 {
-#if PLATFORM_LUMIN
 	if (!DefaultStereoLayers.IsValid())
 	{
 		TSharedPtr<FMagicLeapStereoLayers, ESPMode::ThreadSafe> NewLayersPtr = FSceneViewExtensions::NewExtension<FMagicLeapStereoLayers>(this);
 		DefaultStereoLayers = StaticCastSharedPtr<FDefaultStereoLayers>(NewLayersPtr);
 	}
-#endif //PLATFORM_LUMIN
 	return DefaultStereoLayers.Get();
 }
 
+FMagicLeapStereoLayers::FMagicLeapStereoLayers(const class FAutoRegister& AutoRegister, class FHeadMountedDisplayBase* InHmd)
+: FDefaultStereoLayers(AutoRegister, InHmd)
+, DefaultX(11.0f)
+, DefaultY(-18.0f)
+, DefaultZ(82.0f)
+, DefaultWidth(75.0f)
+, DefaultHeight(40.0f)
+, CVarX(nullptr)
+, CVarY(nullptr)
+, CVarZ(nullptr)
+, CVarWidth(nullptr)
+, CVarHeight(nullptr)
+{
+	CVarX = IConsoleManager::Get().FindConsoleVariable(TEXT("r.LuminDebugCanvasX"));
+	CVarY = IConsoleManager::Get().FindConsoleVariable(TEXT("r.LuminDebugCanvasY"));
+	CVarZ = IConsoleManager::Get().FindConsoleVariable(TEXT("r.LuminDebugCanvasZ"));
+	CVarWidth = IConsoleManager::Get().FindConsoleVariable(TEXT("r.LuminDebugCanvasWidth"));
+	CVarHeight = IConsoleManager::Get().FindConsoleVariable(TEXT("r.LuminDebugCanvasHeight"));
+}
+
 IStereoLayers::FLayerDesc FMagicLeapStereoLayers::GetDebugCanvasLayerDesc(FTextureRHIRef Texture)
 {
 	IStereoLayers::FLayerDesc StereoLayerDesc;
-	StereoLayerDesc.Transform = FTransform(FVector(110.f, 0.f, 0.f));
+	const float DebugCanvasX = CVarX ? CVarX->GetFloat() : DefaultX;
+	const float DebugCanvasY = CVarY ? CVarY->GetFloat() : DefaultY;
+	const float DebugCanvasZ = CVarZ ? CVarZ->GetFloat() : DefaultZ;
+	StereoLayerDesc.Transform = FTransform(FVector(DebugCanvasZ, DebugCanvasX, DebugCanvasY));
 	if (IsOpenGLPlatform(GMaxRHIShaderPlatform))
 	{
 		StereoLayerDesc.Transform.SetScale3D(FVector(1.f, 1.f, -1.f));
 	}
-	StereoLayerDesc.QuadSize = FVector2D(75.f, 40.f);
+	const float DebugCanvasWidth = CVarWidth ? CVarWidth->GetFloat() : DefaultWidth;
+	const float DebugCanvasHeight = CVarHeight ? CVarHeight->GetFloat() : DefaultHeight;
+	StereoLayerDesc.QuadSize = FVector2D(DebugCanvasWidth, DebugCanvasHeight);
 	StereoLayerDesc.PositionType = IStereoLayers::ELayerType::FaceLocked;
 	StereoLayerDesc.ShapeType = IStereoLayers::ELayerShape::QuadLayer;
 	StereoLayerDesc.Texture = Texture;
diff --git a/Engine/Plugins/Lumin/MagicLeap/Source/MagicLeap/Private/MagicLeapStereoLayers.h b/Engine/Plugins/Lumin/MagicLeap/Source/MagicLeap/Private/MagicLeapStereoLayers.h
--- a/Engine/Plugins/Lumin/MagicLeap/Source/MagicLeap/Private/MagicLeapStereoLayers.h
+++ b/Engine/Plugins/Lumin/MagicLeap/Source/MagicLeap/Private/MagicLeapStereoLayers.h
@@ -5,12 +5,26 @@
 
 #include "DefaultStereoLayers.h"
 
+class IConsoleVariable;
+
 class FMagicLeapStereoLayers : public FDefaultStereoLayers
 {
 public:
-	FMagicLeapStereoLayers(const class FAutoRegister& AutoRegister, class FHeadMountedDisplayBase* InHmd) : FDefaultStereoLayers(AutoRegister, InHmd) {}
+	FMagicLeapStereoLayers(const class FAutoRegister& AutoRegister, class FHeadMountedDisplayBase* InHmd);
 
 public:
 	//~ IStereoLayers interface
 	virtual FLayerDesc GetDebugCanvasLayerDesc(FTextureRHIRef Texture) override;
+
+protected:
+	const float DefaultX;
+	const float DefaultY;
+	const float DefaultZ;
+	const float DefaultWidth;
+	const float DefaultHeight;
+	IConsoleVariable* CVarX;
+	IConsoleVariable* CVarY;
+	IConsoleVariable* CVarZ;
+	IConsoleVariable* CVarWidth;
+	IConsoleVariable* CVarHeight;
 };
diff --git a/Engine/Plugins/Lumin/MagicLeap/Source/MagicLeap/Private/MeshTrackerComponent.cpp b/Engine/Plugins/Lumin/MagicLeap/Source/MagicLeap/Private/MeshTrackerComponent.cpp
--- a/Engine/Plugins/Lumin/MagicLeap/Source/MagicLeap/Private/MeshTrackerComponent.cpp
+++ b/Engine/Plugins/Lumin/MagicLeap/Source/MagicLeap/Private/MeshTrackerComponent.cpp
@@ -25,8 +25,11 @@ ML_INCLUDES_END
 #endif //WITH_MLSDK
 
 #if WITH_MLSDK
+// TODO: Don't rely on the size being same.
+static_assert(sizeof(FGuid) == sizeof(MLCoordinateFrameUID), "Size of FGuid should be same as MLCoordinateFrameUID. TODO: Don't rely on the size being same.");
+
 // Map an Unreal meshing LOD to the corresponding ML meshing LOD
-FORCEINLINE MLMeshingLOD MLToUnreal_MeshLOD(EMeshLOD UnrealMeshLOD)
+FORCEINLINE MLMeshingLOD UnrealToML_MeshLOD(EMeshLOD UnrealMeshLOD)
 {
 	switch (UnrealMeshLOD)
 	{
@@ -41,29 +44,50 @@ FORCEINLINE MLMeshingLOD MLToUnreal_MeshLOD(EMeshLOD UnrealMeshLOD)
 	return MLMeshingLOD_Minimum;
 }
 
-// Makes MLCoordinateFrameUID a hashable type for use with TMap and TSet
-template<typename ValueType> struct TMLCoordinateFrameUIDKeyFuncs : 
-	BaseKeyFuncs<TPair<MLCoordinateFrameUID, ValueType>, MLCoordinateFrameUID>
+EMeshState MLToUEMeshState(MLMeshingMeshState MLMeshState)
 {
-private:
-	typedef BaseKeyFuncs<TPair<MLCoordinateFrameUID, ValueType>, MLCoordinateFrameUID> Super;
-public:
-	typedef typename Super::ElementInitType ElementInitType;
-	typedef typename Super::KeyInitType KeyInitType;
-
-	static KeyInitType GetSetKey(ElementInitType Element)
+	switch (MLMeshState)
 	{
-		return Element.Key;
+		case MLMeshingMeshState_New:
+			return EMeshState::New;
+		case MLMeshingMeshState_Updated:
+			return EMeshState::Updated;
+		case MLMeshingMeshState_Deleted:
+			return EMeshState::Deleted;
+		case MLMeshingMeshState_Unchanged:
+			return EMeshState::Unchanged;
 	}
-	static bool Matches(KeyInitType A, KeyInitType B)
-	{
-		return A.data[0] == B.data[0] && A.data[1] == B.data[1];
-	}
-	static uint32 GetKeyHash(KeyInitType Key)
+	check(false);
+	return EMeshState::Unchanged;
+}
+
+void MLToUnrealBlockInfo(const MLMeshingBlockInfo& MLBlockInfo, const FTransform& TrackingToWorld, float WorldToMetersScale, FMeshBlockInfo& UEBlockInfo)
+{
+	FMemory::Memcpy(&UEBlockInfo.BlockID, &MLBlockInfo.id, sizeof(MLCoordinateFrameUID));
+
+	FTransform BlockTransform = FTransform(MagicLeap::ToFQuat(MLBlockInfo.extents.rotation), MagicLeap::ToFVector(MLBlockInfo.extents.center, WorldToMetersScale), FVector(1.0f, 1.0f, 1.0f));
+	if (!BlockTransform.GetRotation().IsNormalized())
 	{
-		return GetTypeHash(static_cast<uint64>(Key.data[0])) ^ ~GetTypeHash(static_cast<uint64>(Key.data[1]));
+		FQuat rotation = BlockTransform.GetRotation();
+		rotation.Normalize();
+		BlockTransform.SetRotation(rotation);
 	}
-};
+
+	BlockTransform.AddToTranslation(TrackingToWorld.GetLocation());
+	BlockTransform.ConcatenateRotation(TrackingToWorld.Rotator().Quaternion());
+	UEBlockInfo.BlockPosition = BlockTransform.GetLocation();
+	UEBlockInfo.BlockOrientation = BlockTransform.Rotator();
+	UEBlockInfo.BlockDimensions = MagicLeap::ToFVectorExtents(MLBlockInfo.extents.extents, WorldToMetersScale);
+
+	UEBlockInfo.Timestamp = FTimespan::FromMicroseconds(MLBlockInfo.timestamp / 1000.0);
+	UEBlockInfo.BlockState = MLToUEMeshState(MLBlockInfo.state);
+}
+
+void UnrealToMLBlockRequest(const FMeshBlockRequest& UEBlockRequest, MLMeshingBlockRequest& MLBlockRequest)
+{
+	FMemory::Memcpy(&MLBlockRequest.id, &UEBlockRequest.BlockID, sizeof(MLCoordinateFrameUID));
+	MLBlockRequest.level = UnrealToML_MeshLOD(UEBlockRequest.LevelOfDetail);
+}
 #endif //WITH_MLSDK
 
 class FMeshTrackerImpl : public MagicLeap::IAppEventHandler
@@ -155,12 +179,15 @@ public:
 	// List of ML mesh block IDs and states
 	TArray<MLMeshingBlockRequest> MeshBlockRequests;
 
-	// Map of ML mesh block IDs to MR Mesh brick IDs
-	TMap<MLCoordinateFrameUID, uint64, FDefaultSetAllocator, 
-		TMLCoordinateFrameUIDKeyFuncs<uint64>> MeshBrickCache;
+	// The set of meshing blocks that have been processed or removed between queries
+	TSet<FGuid> MeshBrickCache;
 
 #endif //WITH_MLSDK
 
+	FMeshInfo LastestMeshInfo;
+	TArray<FMeshBlockRequest> UEMeshBlockRequests;
+	TScriptInterface<IMeshBlockSelectorInterface> BlockSelector;
+
 	// Keep a copy of the mesh data here.  MRMeshComponent will use it from the game and render thread.
 	struct FMLCachedMeshData
 	{
@@ -168,7 +195,7 @@ public:
 		
 		FMeshTrackerImpl* Owner = nullptr;
 		
-		IMRMesh::FBrickId BrickId = 0;
+		FGuid BlockID;
 		TArray<FVector> OffsetVertices;
 		TArray<FVector> WorldVertices;
 		TArray<uint32> Triangles;
@@ -184,7 +211,7 @@ public:
 			FMeshTrackerImpl* TempOwner = Owner;
 			Owner = nullptr;
 	
-			BrickId = 0;
+			BlockID.Invalidate();
 			OffsetVertices.Reset();
 			WorldVertices.Reset();
 			Triangles.Reset();
@@ -273,6 +300,40 @@ public:
 		return true;
 	}
 
+	bool Update(const UMeshTrackerComponent& MeshTrackerComponent)
+	{
+#if WITH_MLSDK
+		MLMeshingSettings MeshSettings = CreateSettings(MeshTrackerComponent);
+
+		if (0 != memcmp(&CurrentMeshSettings, &MeshSettings, sizeof(MeshSettings)))
+		{
+			auto Result = MLMeshingUpdateSettings(MeshTracker, &MeshSettings);
+
+			if (MLResult_Ok == Result)
+			{
+				// For some parameter changes we will want to clear already-generated data
+				if (MeshTrackerComponent.MRMesh != nullptr)
+				{
+					if ((MLMeshingFlags_PointCloud & MeshSettings.flags) !=
+						(MLMeshingFlags_PointCloud & CurrentMeshSettings.flags))
+					{
+						UE_LOG(LogMagicLeap, Log,
+							TEXT("MLMeshingSettings change caused a clear"));
+						MeshTrackerComponent.MRMesh->Clear();
+					}
+				}
+
+				CurrentMeshSettings = MeshSettings;
+				return true;
+			}
+
+			UE_LOG(LogMagicLeap, Error,
+				TEXT("MLMeshingUpdateSettings failed: %s"), UTF8_TO_TCHAR(MLGetResultString(Result)));
+		}
+#endif //WITH_MLSDK
+		return false;
+	}
+
 	void Destroy()
 	{
 #if WITH_MLSDK
@@ -297,9 +358,8 @@ private:
 	FCriticalSection FreeCachedMeshDatasMutex; //The free list may be pushed/popped from multiple threads.
 };
 
-UMeshTrackerComponent::UMeshTrackerComponent(const FObjectInitializer& ObjectInitializer)
-	: Super(ObjectInitializer)
-	, VertexColorFromConfidenceZero(FLinearColor::Red)
+UMeshTrackerComponent::UMeshTrackerComponent()
+	: VertexColorFromConfidenceZero(FLinearColor::Red)
 	, VertexColorFromConfidenceOne(FLinearColor::Blue)
 	, Impl(new FMeshTrackerImpl())
 {
@@ -397,25 +457,10 @@ void UMeshTrackerComponent::PostEditChangeProperty(FPropertyChangedEvent& e)
 #if WITH_MLSDK
 	if (MLHandleIsValid(Impl->MeshTracker) && e.Property != nullptr)
 	{
-		MLMeshingSettings MeshSettings = Impl->CreateSettings(*this);
-
-		// Just brute compare
-		if (0 != memcmp(&Impl->CurrentMeshSettings, &MeshSettings, sizeof(MeshSettings)))
+		if (Impl->Update(*this))
 		{
 			UE_LOG(LogMagicLeap, Log, 
 				TEXT("PostEditChangeProperty is changing MLMeshingSettings"));
-
-			auto Result = MLMeshingUpdateSettings(Impl->MeshTracker, &Impl->CurrentMeshSettings);
-
-			if (MLResult_Ok != Result)
-			{
-				UE_LOG(LogMagicLeap, Error, 
-					TEXT("MLMeshingUpdateSettings failed: %s"), UTF8_TO_TCHAR(MLGetResultString(Result)));
-			}
-			else
-			{
-				Impl->CurrentMeshSettings = MeshSettings;
-			}
 		}
 	}
 #endif //WITH_MLSDK
@@ -451,361 +496,496 @@ void UMeshTrackerComponent::TickComponent(float DeltaTime, enum ELevelTick TickT
 		return;
 	}
 
+	// Dont use the bool() operator from TScriptInterface class since it only checks for the InterfacePointer. 
+	// Since the InterfacePointer is null for it's blueprint implementors, bool() operator gives us the wrong result for checking if interface is valid.
+	if (Impl->BlockSelector.GetObject() == nullptr)
+	{
+		UE_LOG(LogMagicLeap, Warning, TEXT("No block selector is connected, using default implementation."));
+		Impl->BlockSelector = this;
+	}
+
 	// Update the bounding box.
 	FTransform PoseInverse = UHeadMountedDisplayFunctionLibrary::GetTrackingToWorldTransform(this).Inverse();
-	PoseInverse.ConcatenateRotation(BoundingVolume->GetComponentQuat());
 	Impl->BoundsCenter = PoseInverse.TransformPosition(BoundingVolume->GetComponentLocation());
-	Impl->BoundsRotation = PoseInverse.GetRotation();
+	Impl->BoundsRotation = PoseInverse.TransformRotation(BoundingVolume->GetComponentQuat());
 
-	float WorldToMetersScale = AppFramework.GetWorldToMetersScale();
+	// Potentially update for changed component parameters
+	if (Impl->Update(*this))
+	{
+		UE_LOG(LogMagicLeap, Log, TEXT("MLMeshingSettings changed on the fly"));
+	}
 
-	// Make sure MR Mesh is at 0,0,0 (verts received from ML meshing are in world space)
+	const float WorldToMetersScale = AppFramework.GetWorldToMetersScale();
+
+	// Make sure MR Mesh is at 0,0,0 (verts received from ML meshing are in tracking space)
 	MRMesh->SendRelativeTransform(FTransform::Identity);
 
 	if (ScanWorld)
 	{
-		// Request mesh info
-		if (Impl->CurrentMeshInfoRequest == ML_INVALID_HANDLE && Impl->CurrentMeshRequest == ML_INVALID_HANDLE)
+		if (GetMeshResult())
 		{
-			MLMeshingExtents Extents = {};
-			Extents.center = MagicLeap::ToMLVector(Impl->BoundsCenter, WorldToMetersScale);
-			Extents.rotation = MagicLeap::ToMLQuat(Impl->BoundsRotation);
-			Extents.extents = MagicLeap::ToMLVector(BoundingVolume->GetScaledBoxExtent(), WorldToMetersScale);
-
-			// MagicLeap::ToMLVector() is, as the name implies, meant for vectors (not extents) and using it
-			// causes the Z component to be negated. We'll abs all of them for safety.
-			Extents.extents.x = FMath::Abs<float>(Extents.extents.x);
-			Extents.extents.y = FMath::Abs<float>(Extents.extents.y);
-			Extents.extents.z = FMath::Abs<float>(Extents.extents.z);
-
-			auto Result = MLMeshingRequestMeshInfo(Impl->MeshTracker, &Extents, &Impl->CurrentMeshInfoRequest);
-			if (MLResult_Ok != Result)
+			RequestMeshInfo();
+			if (GetMeshInfoResult())
 			{
-				UE_LOG(LogMagicLeap, Error, TEXT("MLMeshingRequestMeshInfo failed: %s."), UTF8_TO_TCHAR(MLGetResultString(Result)));
-				Impl->CurrentMeshInfoRequest = ML_INVALID_HANDLE;
+				RequestMesh();
 			}
 		}
+	}
+#endif //WITH_MLSDK
+}
 
-		// Request block meshes for current mesh info and block list
-		if (Impl->CurrentMeshRequest == ML_INVALID_HANDLE && Impl->MeshBlockRequests.Num() > 0)
+void UMeshTrackerComponent::SelectMeshBlocks_Implementation(const FMeshInfo& NewMeshInfo, TArray<FMeshBlockRequest>& RequestedMesh)
+{
+	for (const FMeshBlockInfo& BlockInfo : NewMeshInfo.BlockData)
+	{
+		if (BlockInfo.BlockState == EMeshState::New || BlockInfo.BlockState == EMeshState::Unchanged)
 		{
-			MLMeshingMeshRequest MeshRequest = {};
-			MeshRequest.request_count = static_cast<int>(Impl->MeshBlockRequests.Num());
-			MeshRequest.data = Impl->MeshBlockRequests.GetData();
-			auto Result = MLMeshingRequestMesh(Impl->MeshTracker, &MeshRequest, &Impl->CurrentMeshRequest);
-			if (MLResult_Ok != Result)
+			FMeshBlockRequest BlockRequest;
+			BlockRequest.BlockID = BlockInfo.BlockID;
+			BlockRequest.LevelOfDetail = LevelOfDetail;
+			RequestedMesh.Add(BlockRequest);
+		}
+	}
+}
+
+void UMeshTrackerComponent::ConnectBlockSelector(TScriptInterface<IMeshBlockSelectorInterface> Selector)
+{
+	if (Impl != nullptr)
+	{
+		// Dont use the bool() operator from TScriptInterface class since it only checks for the InterfacePointer. 
+		// Since the InterfacePointer is null for it's blueprint implementors, bool() operator gives us the wrong result for checking if interface is valid.
+		if (Selector.GetObject() != nullptr)
+		{
+			// If called via C++, Selector might have been created manually and not implement IMeshBlockSelectorInterface.
+			if (Selector.GetObject()->GetClass()->ImplementsInterface(UMeshBlockSelectorInterface::StaticClass()))
 			{
-				UE_LOG(LogMagicLeap, Error, TEXT("MLMeshingRequestMesh failed: %s."), UTF8_TO_TCHAR(MLGetResultString(Result)));
-				Impl->CurrentMeshInfoRequest = ML_INVALID_HANDLE;
+				Impl->BlockSelector = Selector;
+			}
+			else
+			{
+				UE_LOG(LogMagicLeap, Warning, TEXT("Selector %s does not implement IMeshBlockSelectorInterface. Using default block selector from MeshTrackerComponent."), *(Selector.GetObject()->GetFName().ToString()));
+				Impl->BlockSelector = this;	
 			}
 		}
+		else
+		{
+			UE_LOG(LogMagicLeap, Warning, TEXT("Invalid selector passed to UMeshTrackerComponent::ConnectBlockSelector(). Using default block selector from MeshTrackerComponent."));
+			Impl->BlockSelector = this;	
+		}
+	}
+}
+
+void UMeshTrackerComponent::DisconnectBlockSelector()
+{
+	if (Impl != nullptr)
+	{
+		Impl->BlockSelector = this;
+	}
+}
+
+void UMeshTrackerComponent::BeginDestroy()
+{
+	if (MRMesh != nullptr)
+	{
+		DisconnectMRMesh(MRMesh);
+	}
+	Super::BeginDestroy();
+}
+
+void UMeshTrackerComponent::FinishDestroy()
+{
+#if WITH_EDITOR
+	if (GIsEditor)
+	{
+		FEditorDelegates::PrePIEEnded.RemoveAll(this);
+	}
+#endif
+	Impl->Destroy();
+	Super::FinishDestroy();
+}
 
-		// Request IDs and states of block meshes
-		if (Impl->CurrentMeshInfoRequest != ML_INVALID_HANDLE)
+void UMeshTrackerComponent::RequestMeshInfo()
+{
+#if WITH_MLSDK
+	const FAppFramework& AppFramework = static_cast<FMagicLeapHMD *>(GEngine->XRSystem->GetHMDDevice())->GetAppFrameworkConst();
+	const float WorldToMetersScale = AppFramework.GetWorldToMetersScale();
+
+	// Request mesh info as frequently as possible.
+	// Actual request for mesh will be submitted based on the latest available info at the time of triggering the mesh request.
+	if (Impl->CurrentMeshInfoRequest == ML_INVALID_HANDLE)
+	{
+		MLMeshingExtents Extents = {};
+		Extents.center = MagicLeap::ToMLVector(Impl->BoundsCenter, WorldToMetersScale);
+		Extents.rotation = MagicLeap::ToMLQuat(Impl->BoundsRotation);
+		Extents.extents = MagicLeap::ToMLVectorExtents(BoundingVolume->GetScaledBoxExtent(), WorldToMetersScale);
+
+		auto Result = MLMeshingRequestMeshInfo(Impl->MeshTracker, &Extents, &Impl->CurrentMeshInfoRequest);
+		if (MLResult_Ok != Result)
 		{
-			MLMeshingMeshInfo MeshInfo = {};
+			UE_LOG(LogMagicLeap, Error, TEXT("MLMeshingRequestMeshInfo failed: %s."), UTF8_TO_TCHAR(MLGetResultString(Result)));
+			Impl->CurrentMeshInfoRequest = ML_INVALID_HANDLE;
+		}
+	}
+#endif
+}
 
-			auto Result = MLMeshingGetMeshInfoResult(Impl->MeshTracker, Impl->CurrentMeshInfoRequest, &MeshInfo);
-			if (MLResult_Ok != Result)
+bool UMeshTrackerComponent::GetMeshInfoResult()
+{
+#if WITH_MLSDK
+	const FAppFramework& AppFramework = static_cast<FMagicLeapHMD *>(GEngine->XRSystem->GetHMDDevice())->GetAppFrameworkConst();
+	const float WorldToMetersScale = AppFramework.GetWorldToMetersScale();
+
+	// Check for updated mesh info and cache the result.
+	// The cached result will be used by app to choose which blocks it wants to actually request the mesh for.
+	if (Impl->CurrentMeshInfoRequest != ML_INVALID_HANDLE)
+	{
+		MLMeshingMeshInfo MeshInfo = {};
+
+		auto Result = MLMeshingGetMeshInfoResult(Impl->MeshTracker, Impl->CurrentMeshInfoRequest, &MeshInfo);
+		if (MLResult_Ok != Result)
+		{
+			// Just silently wait for pending result
+			if (MLResult_Pending != Result)
 			{
-				// Just silently wait for pending result
-				if (MLResult_Pending != Result)
-				{
-					UE_LOG(LogMagicLeap, Error, 
-						TEXT("MLMeshingGetMeshInfoResult failed: %s."), UTF8_TO_TCHAR(MLGetResultString(Result)));
-				}
+				UE_LOG(LogMagicLeap, Error, 
+					TEXT("MLMeshingGetMeshInfoResult failed: %s."), UTF8_TO_TCHAR(MLGetResultString(Result)));
+				Impl->CurrentMeshInfoRequest = ML_INVALID_HANDLE;
+				return true;
 			}
-			else
+			return false;
+		}
+		else
+		{
+			// Clear our stored block requests
+			Impl->LastestMeshInfo.BlockData.Empty(MeshInfo.data_count);
+			Impl->LastestMeshInfo.BlockData.AddUninitialized(MeshInfo.data_count);
+			Impl->LastestMeshInfo.Timestamp = FTimespan::FromMicroseconds(MeshInfo.timestamp / 1000.0);
+			UE_LOG(LogMagicLeap, Warning, TEXT("Mesh info timestamp = %f"), Impl->LastestMeshInfo.Timestamp.GetTotalMilliseconds());
+
+			for (uint32_t MeshInfoIndex = 0; MeshInfoIndex < MeshInfo.data_count; ++ MeshInfoIndex)
 			{
-				// Clear our stored block requests
-				Impl->MeshBlockRequests.Empty();
+				const auto &MeshInfoData = MeshInfo.data[MeshInfoIndex];
 
-				//UE_LOG(LogMagicLeap, Log, TEXT("Mesh tracker received %d mesh block infos"), MeshInfo.data_count);
+				// TODO: right now we are adding even the deleted and unchanged blocks here. we probably only need to add new and unchanged blocks.
+				MLToUnrealBlockInfo(MeshInfoData, UHeadMountedDisplayFunctionLibrary::GetTrackingToWorldTransform(this), WorldToMetersScale, Impl->LastestMeshInfo.BlockData[MeshInfoIndex]);
 
-				for (uint32_t MeshInfoIndex = 0; MeshInfoIndex < MeshInfo.data_count; ++ MeshInfoIndex)
+				switch (MeshInfoData.state)
 				{
-					const auto &MeshInfoData = MeshInfo.data[MeshInfoIndex];
+					// TODO: maybe do this only when the rest of the mesh is updated i.e. success from MLMeshingGetMeshResult()
+					case MLMeshingMeshState_Deleted:
+					{
+						const auto& BlockID = FGuid(MeshInfoData.id.data[0], MeshInfoData.id.data[0] >> 32, MeshInfoData.id.data[1], MeshInfoData.id.data[1] >> 32);
 
-					MLMeshingBlockRequest BlockRequest = {};
+						// Delete the brick and its cache entry
+						if (Impl->MeshBrickCache.Contains(BlockID))
+						{
 
-					BlockRequest.id = MeshInfoData.id;
-					BlockRequest.level = MLToUnreal_MeshLOD(LevelOfDetail);
+							if (MeshType != EMeshType::PointCloud)
+							{
+								const static TArray<FVector> EmptyVertices;
+								const static TArray<FVector2D> EmptyUVs;
+								const static TArray<FPackedNormal> EmptyTangents;
+								const static TArray<FColor> EmptyVertexColors;
+								const static TArray<uint32> EmptyTriangles;
+								static_cast<IMRMesh*>(MRMesh)->SendBrickData(IMRMesh::FSendBrickDataArgs
+									{
+										nullptr,
+										(static_cast<uint64>(BlockID.B) << 32 | BlockID.A),
+										EmptyVertices,
+										EmptyUVs,
+										EmptyTangents,
+										EmptyVertexColors,
+										EmptyTriangles
+									}
+								);
+							}
 
-					switch (MeshInfoData.state)
-					{
-						case MLMeshingMeshState_New:
-						case MLMeshingMeshState_Updated:
-						{
-							// Store the block request so we can update it
-							Impl->MeshBlockRequests.Add(BlockRequest);
-							break;
-						}
-						case MLMeshingMeshState_Deleted:
-						{
-							// Delete the brick and its cache entry
-							if (Impl->MeshBrickCache.Contains(MeshInfoData.id))
+							if (OnMeshTrackerUpdated.IsBound())
 							{
-								const auto& BrickId = Impl->MeshBrickCache[MeshInfoData.id];
-
-								if (MeshType != EMeshType::PointCloud)
-								{
-									const static TArray<FVector> EmptyVertices;
-									const static TArray<FVector2D> EmptyUVs;
-									const static TArray<FPackedNormal> EmptyTangents;
-									const static TArray<FColor> EmptyVertexColors;
-									const static TArray<uint32> EmptyTriangles;
-									static_cast<IMRMesh*>(MRMesh)->SendBrickData(IMRMesh::FSendBrickDataArgs
-										{
-											nullptr,
-											BrickId,
-											EmptyVertices,
-											EmptyUVs,
-											EmptyTangents,
-											EmptyVertexColors,
-											EmptyTriangles
-										}
-									);
-								}
-
-								if (OnMeshTrackerUpdated.IsBound())
-								{
-									OnMeshTrackerUpdated.Broadcast(BrickId, 
-										TArray<FVector>(), TArray<int32>(), TArray<FVector>(), TArray<float>());
-								}
-
-								Impl->MeshBrickCache.Remove(MeshInfoData.id);
+								OnMeshTrackerUpdated.Broadcast(BlockID,
+									TArray<FVector>(), TArray<int32>(), TArray<FVector>(), TArray<float>());
 							}
-							break;
+
+							Impl->MeshBrickCache.Remove(BlockID);
 						}
-						default:
-							break;
+						break;
 					}
+					default:
+						break;
 				}
+			}
+
+			// Free up the ML meshing resources
+			MLMeshingFreeResource(Impl->MeshTracker, &Impl->CurrentMeshInfoRequest);
+			Impl->CurrentMeshInfoRequest = ML_INVALID_HANDLE;
+			return true;
+		}
+	}
+#endif
+
+	// if it reaches here, something has gone wrong
+	return false;
+}
+
+void UMeshTrackerComponent::RequestMesh()
+{
+#if WITH_MLSDK
+	// Request block meshes for current mesh info and block list
+	if (Impl->CurrentMeshRequest == ML_INVALID_HANDLE)
+	{
+		Impl->UEMeshBlockRequests.Empty(Impl->LastestMeshInfo.BlockData.Num());
+		// Allow applications to choose which blocks to mesh.
+		IMeshBlockSelectorInterface::Execute_SelectMeshBlocks(Impl->BlockSelector.GetObject(), Impl->LastestMeshInfo, Impl->UEMeshBlockRequests);
 
-				// Free up the ML meshing resources
-				MLMeshingFreeResource(Impl->MeshTracker, &Impl->CurrentMeshInfoRequest);
+		if (Impl->UEMeshBlockRequests.Num() > 0)
+		{
+			Impl->MeshBlockRequests.Empty(Impl->UEMeshBlockRequests.Num());
+			Impl->MeshBlockRequests.AddUninitialized(Impl->UEMeshBlockRequests.Num());
+			for (int32 i = 0; i < Impl->UEMeshBlockRequests.Num(); ++i)
+			{
+				UnrealToMLBlockRequest(Impl->UEMeshBlockRequests[i], Impl->MeshBlockRequests[i]);
+			}
+
+			MLMeshingMeshRequest MeshRequest = {};
+			MeshRequest.request_count = static_cast<int>(Impl->MeshBlockRequests.Num());
+			MeshRequest.data = Impl->MeshBlockRequests.GetData();
+			auto Result = MLMeshingRequestMesh(Impl->MeshTracker, &MeshRequest, &Impl->CurrentMeshRequest);
+			if (MLResult_Ok != Result)
+			{
+				UE_LOG(LogMagicLeap, Error, TEXT("MLMeshingRequestMesh failed: %s."), UTF8_TO_TCHAR(MLGetResultString(Result)));
 				Impl->CurrentMeshInfoRequest = ML_INVALID_HANDLE;
 			}
 		}
+	}
+#endif
+}
 
-		// Request mesh data
-		if (Impl->CurrentMeshRequest != ML_INVALID_HANDLE)
-		{
-			MLMeshingMesh Mesh = {};
+bool UMeshTrackerComponent::GetMeshResult()
+{
+#if WITH_MLSDK
+	const FAppFramework& AppFramework = static_cast<FMagicLeapHMD *>(GEngine->XRSystem->GetHMDDevice())->GetAppFrameworkConst();
+	const float WorldToMetersScale = AppFramework.GetWorldToMetersScale();
 
-			auto Result = MLMeshingGetMeshResult(Impl->MeshTracker, Impl->CurrentMeshRequest, &Mesh);
+	// Get mesh result
+	if (Impl->CurrentMeshRequest != ML_INVALID_HANDLE)
+	{
+		MLMeshingMesh Mesh = {};
 
-			if (MLResult_Ok != Result)
+		auto Result = MLMeshingGetMeshResult(Impl->MeshTracker, Impl->CurrentMeshRequest, &Mesh);
+
+		if (MLResult_Ok != Result)
+		{
+			// Just silently wait for pending result
+			if (MLResult_Pending != Result)
 			{
-				// Just silently wait for pending result
-				if (MLResult_Pending != Result)
-				{
-					UE_LOG(LogMagicLeap, Error, TEXT("MLMeshingGetMeshResult failed: %s."), UTF8_TO_TCHAR(MLGetResultString(Result)));
-				}
+				UE_LOG(LogMagicLeap, Error, TEXT("MLMeshingGetMeshResult failed: %s."), UTF8_TO_TCHAR(MLGetResultString(Result)));
+				// Mesh request failed, lets queue another one.
+				Impl->CurrentMeshRequest = ML_INVALID_HANDLE;
+				return true;
 			}
-			else
+			// Mesh request pending...
+			return false;
+		}
+		else
+		{
+			FVector VertexOffset = UHeadMountedDisplayFunctionLibrary::GetTrackingToWorldTransform(this).Inverse().GetLocation();
+			for (uint32_t MeshIndex = 0; MeshIndex < Mesh.data_count; ++ MeshIndex)
 			{
-				//UE_LOG(LogMagicLeap, Log, TEXT("Mesh tracker received %d blocks for mesh request %d"), 
-					//Mesh.data_count, static_cast<int>(Impl->CurrentMeshRequest));
+				const auto &MeshData = Mesh.data[MeshIndex];
 
-				FVector VertexOffset = UHeadMountedDisplayFunctionLibrary::GetTrackingToWorldTransform(this).Inverse().GetLocation();
-				for (uint32_t MeshIndex = 0; MeshIndex < Mesh.data_count; ++ MeshIndex)
-				{
-					const auto &MeshData = Mesh.data[MeshIndex];
+				auto BlockID = FGuid(MeshData.id.data[0], MeshData.id.data[0] >> 32, MeshData.id.data[1], MeshData.id.data[1] >> 32);
 
-					// Create a brick index for any new mesh block
-					if (!Impl->MeshBrickCache.Contains(MeshData.id))
-					{
-						Impl->MeshBrickCache.Add(MeshData.id, Impl->MeshBrickIndex ++);
-					}
+				// Create a brick index for any new mesh block
+				if (!Impl->MeshBrickCache.Contains(BlockID))
+				{
+					Impl->MeshBrickCache.Add(BlockID);
+				}
 
-					auto BrickId = Impl->MeshBrickCache[MeshData.id];
+				// Acquire mesh data cache and mark its brick ID
+				FMeshTrackerImpl::FMLCachedMeshData::SharedPtr CurrentMeshDataCache = Impl->AquireMeshDataCache();
+				CurrentMeshDataCache->BlockID = BlockID;
 
-					// Acquire mesh data cache and mark its brick ID
-					FMeshTrackerImpl::FMLCachedMeshData::SharedPtr CurrentMeshDataCache = Impl->AquireMeshDataCache();
-					CurrentMeshDataCache->BrickId = BrickId;
+				// Pull vertices
+				CurrentMeshDataCache->OffsetVertices.Reserve(MeshData.vertex_count);
+				CurrentMeshDataCache->WorldVertices.Reserve(MeshData.vertex_count);
+				for (uint32_t v = 0; v < MeshData.vertex_count; ++ v)
+				{
+					CurrentMeshDataCache->OffsetVertices.Add(MagicLeap::ToFVector(MeshData.vertex[v], WorldToMetersScale) - VertexOffset);
+					CurrentMeshDataCache->WorldVertices.Add(MagicLeap::ToFVector(MeshData.vertex[v], WorldToMetersScale));
+				}
 
-					// Pull vertices
-					CurrentMeshDataCache->OffsetVertices.Reserve(MeshData.vertex_count);
-					CurrentMeshDataCache->WorldVertices.Reserve(MeshData.vertex_count);
-					for (uint32_t v = 0; v < MeshData.vertex_count; ++ v)
-					{
-							CurrentMeshDataCache->OffsetVertices.Add(MagicLeap::ToFVector(MeshData.vertex[v], WorldToMetersScale) - VertexOffset);
-							CurrentMeshDataCache->WorldVertices.Add(MagicLeap::ToFVector(MeshData.vertex[v], WorldToMetersScale));
-					}
+				// Pull indices
+				CurrentMeshDataCache->Triangles.Reserve(MeshData.index_count);
+				for (uint16_t i = 0; i < MeshData.index_count; ++ i)
+				{
+					CurrentMeshDataCache->Triangles.Add(static_cast<uint32>(MeshData.index[i]));
+				}
 
-					// Pull indices
-					CurrentMeshDataCache->Triangles.Reserve(MeshData.index_count);
-					for (uint16_t i = 0; i < MeshData.index_count; ++ i)
+				// Pull normals
+				CurrentMeshDataCache->Normals.Reserve(MeshData.vertex_count);
+				if (nullptr != MeshData.normal)
+				{
+					for (uint32_t n = 0; n < MeshData.vertex_count; ++ n)
 					{
-						CurrentMeshDataCache->Triangles.Add(static_cast<uint32>(MeshData.index[i]));
+						CurrentMeshDataCache->Normals.Add(MagicLeap::ToFVectorNoScale(MeshData.normal[n]));
 					}
-
-					// Pull normals
-					CurrentMeshDataCache->Normals.Reserve(MeshData.vertex_count);
-					if (nullptr != MeshData.normal)
+				}
+				// If no normals were provided we need to pack fake ones for Vulkan
+				else
+				{
+					for (uint32_t n = 0; n < MeshData.vertex_count; ++ n)
 					{
-						for (uint32_t n = 0; n < MeshData.vertex_count; ++ n)
-						{
-							CurrentMeshDataCache->Normals.Add(MagicLeap::
-								ToFVector(MeshData.normal[n], 1.0f));
-						}
-					}
-					// If no normals were provided we need to pack fake ones for Vulkan
-					else
-					{
-						for (uint32_t n = 0; n < MeshData.vertex_count; ++ n)
-						{
-							FVector FakeNormal = CurrentMeshDataCache->OffsetVertices[n];
-							FakeNormal.Normalize();
-							CurrentMeshDataCache->Normals.Add(FakeNormal);
-						}
+						FVector FakeNormal = CurrentMeshDataCache->OffsetVertices[n];
+						FakeNormal.Normalize();
+						CurrentMeshDataCache->Normals.Add(FakeNormal);
 					}
+				}
 
-					// Calculate and pack tangents
-					CurrentMeshDataCache->Tangents.Reserve(MeshData.vertex_count * 2);
-					for (uint32_t t = 0; t < MeshData.vertex_count; ++ t)
-					{
-						const FVector& Norm = CurrentMeshDataCache->Normals[t];
+				// Calculate and pack tangents
+				CurrentMeshDataCache->Tangents.Reserve(MeshData.vertex_count * 2);
+				for (uint32_t t = 0; t < MeshData.vertex_count; ++ t)
+				{
+					const FVector& Norm = CurrentMeshDataCache->Normals[t];
 
-						// Calculate tangent
-						auto Perp = Norm.X < Norm.Z ? 
-							FVector(1.0f, 0.0f, 0.0f) : FVector(0.0f, 1.0f, 0.0f);
-						auto Tang = FVector::CrossProduct(Norm, Perp);
+					// Calculate tangent
+					auto Perp = Norm.X < Norm.Z ? 
+						FVector(1.0f, 0.0f, 0.0f) : FVector(0.0f, 1.0f, 0.0f);
+					auto Tang = FVector::CrossProduct(Norm, Perp);
 
-						CurrentMeshDataCache->Tangents.Add(Tang);
-						CurrentMeshDataCache->Tangents.Add(Norm);
-					}
+					CurrentMeshDataCache->Tangents.Add(Tang);
+					CurrentMeshDataCache->Tangents.Add(Norm);
+				}
 
-					// Pull confidence
-					if (nullptr != MeshData.confidence)
-					{
-						CurrentMeshDataCache->Confidence.Append(MeshData.confidence, MeshData.vertex_count);
-					}
+				// Pull confidence
+				if (nullptr != MeshData.confidence)
+				{
+					CurrentMeshDataCache->Confidence.Append(MeshData.confidence, MeshData.vertex_count);
+				}
 
-					// Apply chosen vertex color mode
-					switch (VertexColorMode)
+				// Apply chosen vertex color mode
+				switch (VertexColorMode)
+				{
+					case EMLMeshVertexColorMode::Confidence:
 					{
-						case EMLMeshVertexColorMode::Confidence:
+						if (nullptr != MeshData.confidence)
 						{
-							if (nullptr != MeshData.confidence)
-							{
-								CurrentMeshDataCache->VertexColors.Reserve(MeshData.vertex_count);
-								for (uint32 v = 0; v < MeshData.vertex_count; ++ v)
-								{
-									const FLinearColor VertexColor = FMath::Lerp(VertexColorFromConfidenceZero, 
-										VertexColorFromConfidenceOne, CurrentMeshDataCache->Confidence[v]);
-									CurrentMeshDataCache->VertexColors.Add(VertexColor.ToFColor(false));
-								}
-							}
-							else
+							CurrentMeshDataCache->VertexColors.Reserve(MeshData.vertex_count);
+							for (uint32 v = 0; v < MeshData.vertex_count; ++ v)
 							{
-								UE_LOG(LogMagicLeap, Warning, TEXT("MeshTracker vertex color mode is Confidence "
-									"but no confidence values available. Using white for all blocks."));
+								const FLinearColor VertexColor = FMath::Lerp(VertexColorFromConfidenceZero, 
+									VertexColorFromConfidenceOne, CurrentMeshDataCache->Confidence[v]);
+								CurrentMeshDataCache->VertexColors.Add(VertexColor.ToFColor(false));
 							}
-							break;
 						}
-						case EMLMeshVertexColorMode::Block:
+						else
 						{
-							if (BlockVertexColors.Num() > 0)
-							{
-								const FColor& VertexColor = BlockVertexColors[BrickId % BlockVertexColors.Num()];
+							UE_LOG(LogMagicLeap, Warning, TEXT("MeshTracker vertex color mode is Confidence "
+								"but no confidence values available. Using white for all blocks."));
+						}
+						break;
+					}
+					case EMLMeshVertexColorMode::Block:
+					{
+						if (BlockVertexColors.Num() > 0)
+						{
+							const FColor& VertexColor = BlockVertexColors[(static_cast<uint64>(BlockID.B) << 32 | BlockID.A) % BlockVertexColors.Num()];
 
-								CurrentMeshDataCache->VertexColors.Reserve(MeshData.vertex_count);
-								for (uint32 v = 0; v < MeshData.vertex_count; ++ v)
-								{
-									CurrentMeshDataCache->VertexColors.Add(VertexColor);
-								}
-							}
-							else
+							CurrentMeshDataCache->VertexColors.Reserve(MeshData.vertex_count);
+							for (uint32 v = 0; v < MeshData.vertex_count; ++ v)
 							{
-								UE_LOG(LogMagicLeap, Warning, TEXT("MeshTracker vertex color mode is Block but "
-									"no BlockVertexColors set. Using white for all blocks."));
+								CurrentMeshDataCache->VertexColors.Add(VertexColor);
 							}
-							break;
 						}
-						case EMLMeshVertexColorMode::None:
+						else
 						{
-							break;
+							UE_LOG(LogMagicLeap, Warning, TEXT("MeshTracker vertex color mode is Block but "
+								"no BlockVertexColors set. Using white for all blocks."));
 						}
-						default:
-							check(false);
-							break;
+						break;
 					}
-
-					// To work in all rendering paths we always set a vertex color
-					if (CurrentMeshDataCache->VertexColors.Num() == 0)
+					case EMLMeshVertexColorMode::LOD:
 					{
-						for (uint32 v = 0; v < MeshData.vertex_count; ++ v)
+						if (BlockVertexColors.Num() >= MLMeshingLOD_Maximum)
+						{
+							const FColor& VertexColor = BlockVertexColors[MeshData.level];
+
+							CurrentMeshDataCache->VertexColors.Reserve(MeshData.vertex_count);
+							for (uint32 v = 0; v < MeshData.vertex_count; ++ v)
+							{
+								CurrentMeshDataCache->VertexColors.Add(VertexColor);
+							}
+						}
+						else
 						{
-							CurrentMeshDataCache->VertexColors.Add(FColor::White);
+							UE_LOG(LogMagicLeap, Warning, TEXT("MeshTracker vertex color mode is LOD but "
+								"BlockVertexColors are less then the number of LODs. Using white for all blocks."));
 						}
+						break;
 					}
-
-					// Write UVs
-					CurrentMeshDataCache->UV0.Reserve(MeshData.vertex_count);
-					for (uint32 v = 0; v < MeshData.vertex_count; ++ v)
+					case EMLMeshVertexColorMode::None:
 					{
-						const float FakeCoord = static_cast<float>(v) / static_cast<float>(MeshData.vertex_count);
-						CurrentMeshDataCache->UV0.Add(FVector2D(FakeCoord, FakeCoord));
+						break;
 					}
+					default:
+						check(false);
+						break;
+				}
 
-					// Create/update brick
-					if (MeshType != EMeshType::PointCloud)
+				// To work in all rendering paths we always set a vertex color
+				if (CurrentMeshDataCache->VertexColors.Num() == 0)
+				{
+					for (uint32 v = 0; v < MeshData.vertex_count; ++ v)
 					{
-						static_cast<IMRMesh*>(MRMesh)->SendBrickData(IMRMesh::FSendBrickDataArgs
-							{
-								TSharedPtr<IMRMesh::FBrickDataReceipt, ESPMode::ThreadSafe>
-									(new FMeshTrackerImpl::FMeshTrackerComponentBrickDataReceipt(CurrentMeshDataCache)),
-								CurrentMeshDataCache->BrickId,
-								CurrentMeshDataCache->WorldVertices,
-								CurrentMeshDataCache->UV0,
-								CurrentMeshDataCache->Tangents,
-								CurrentMeshDataCache->VertexColors,
-								CurrentMeshDataCache->Triangles
-							}
-						);
+						CurrentMeshDataCache->VertexColors.Add(FColor::White);
 					}
+				}
 
-					// Broadcast that a mesh was updated
-					if (OnMeshTrackerUpdated.IsBound())
-					{
-						// Hack because blueprints don't support uint32.
-						TArray<int32> Triangles(reinterpret_cast<const int32*>(CurrentMeshDataCache->
-							Triangles.GetData()), CurrentMeshDataCache->Triangles.Num());
-						OnMeshTrackerUpdated.Broadcast(CurrentMeshDataCache->BrickId, CurrentMeshDataCache->OffsetVertices, 
-							Triangles, CurrentMeshDataCache->Normals, CurrentMeshDataCache->Confidence);
-					}
+				// Write UVs
+				CurrentMeshDataCache->UV0.Reserve(MeshData.vertex_count);
+				for (uint32 v = 0; v < MeshData.vertex_count; ++ v)
+				{
+					const float FakeCoord = static_cast<float>(v) / static_cast<float>(MeshData.vertex_count);
+					CurrentMeshDataCache->UV0.Add(FVector2D(FakeCoord, FakeCoord));
 				}
 
-				// All meshes pulled and/or updated; free the ML resource
-				MLMeshingFreeResource(Impl->MeshTracker, &Impl->CurrentMeshRequest);
-				Impl->CurrentMeshRequest = ML_INVALID_HANDLE;
+				// Create/update brick
+				if (MeshType != EMeshType::PointCloud)
+				{
+					static_cast<IMRMesh*>(MRMesh)->SendBrickData(IMRMesh::FSendBrickDataArgs
+						{
+							TSharedPtr<IMRMesh::FBrickDataReceipt, ESPMode::ThreadSafe>
+								(new FMeshTrackerImpl::FMeshTrackerComponentBrickDataReceipt(CurrentMeshDataCache)),
+							(static_cast<uint64>(CurrentMeshDataCache->BlockID.B) << 32 | CurrentMeshDataCache->BlockID.A),
+							CurrentMeshDataCache->WorldVertices,
+							CurrentMeshDataCache->UV0,
+							CurrentMeshDataCache->Tangents,
+							CurrentMeshDataCache->VertexColors,
+							CurrentMeshDataCache->Triangles
+						}
+					);
+				}
+
+				// Broadcast that a mesh was updated
+				if (OnMeshTrackerUpdated.IsBound())
+				{
+					// Hack because blueprints don't support uint32.
+					TArray<int32> Triangles(reinterpret_cast<const int32*>(CurrentMeshDataCache->
+						Triangles.GetData()), CurrentMeshDataCache->Triangles.Num());
+					OnMeshTrackerUpdated.Broadcast(CurrentMeshDataCache->BlockID, CurrentMeshDataCache->OffsetVertices,
+						Triangles, CurrentMeshDataCache->Normals, CurrentMeshDataCache->Confidence);
+				}
 			}
-		}
-	}
-#endif //WITH_MLSDK
-}
 
-void UMeshTrackerComponent::BeginDestroy()
-{
-	if (MRMesh != nullptr)
-	{
-		DisconnectMRMesh(MRMesh);
-	}
-	Super::BeginDestroy();
-}
+			// All meshes pulled and/or updated; free the ML resource
+			MLMeshingFreeResource(Impl->MeshTracker, &Impl->CurrentMeshRequest);
+			Impl->CurrentMeshRequest = ML_INVALID_HANDLE;
 
-void UMeshTrackerComponent::FinishDestroy()
-{
-#if WITH_EDITOR
-	if (GIsEditor)
-	{
-		FEditorDelegates::PrePIEEnded.RemoveAll(this);
+			return true;
+		}
 	}
 #endif
-	Impl->Destroy();
-	Super::FinishDestroy();
+
+	return true;
 }
 
 #if WITH_EDITOR
diff --git a/Engine/Plugins/Lumin/MagicLeap/Source/MagicLeap/Private/PlanesComponent.cpp b/Engine/Plugins/Lumin/MagicLeap/Source/MagicLeap/Private/PlanesComponent.cpp
--- a/Engine/Plugins/Lumin/MagicLeap/Source/MagicLeap/Private/PlanesComponent.cpp
+++ b/Engine/Plugins/Lumin/MagicLeap/Source/MagicLeap/Private/PlanesComponent.cpp
@@ -161,7 +161,6 @@ UPlanesComponent::UPlanesComponent()
 	, MaxResults(10)
 	, MinHolePerimeter(50.0f)
 	, MinPlaneArea(25.0f)
-	, IgnoreBoundingVolume(false)
 	, Impl(new FPlanesTrackerImpl())
 {
 	// Make sure this component ticks
@@ -194,12 +193,18 @@ UPlanesComponent::~UPlanesComponent()
 	Impl = nullptr;
 }
 
+void UPlanesComponent::BeginPlay()
+{
+	Super::BeginPlay();
+	Impl->Create();
+}
+
 void UPlanesComponent::TickComponent(float DeltaTime, enum ELevelTick TickType, FActorComponentTickFunction *ThisTickFunction)
 {
 #if WITH_MLSDK
 	Super::TickComponent(DeltaTime, TickType, ThisTickFunction);
 
-	if (!(IMagicLeapPlugin::Get().IsMagicLeapHMDValid() && Impl->Create()))
+	if (!(IMagicLeapPlugin::Get().IsMagicLeapHMDValid()))
 	{
 		return;
 	}
@@ -252,7 +257,7 @@ void UPlanesComponent::TickComponent(float DeltaTime, enum ELevelTick TickType,
 					planeTransform.SetRotation(rotation);
 				}
 
-				planeTransform.SetRotation(MagicLeap::ToUERotator(planeTransform.GetRotation()));
+				planeTransform.ConcatenateRotation(FQuat(FVector(0, 0, 1), PI));
 				planeTransform.AddToTranslation(PoseTransform.GetLocation());
 				planeTransform.ConcatenateRotation(PoseTransform.Rotator().Quaternion());
 				resultPlane.PlanePosition = planeTransform.GetLocation();
@@ -261,6 +266,7 @@ void UPlanesComponent::TickComponent(float DeltaTime, enum ELevelTick TickType,
 				// We are rotating it by 90 degrees clock-wise about the right axis (Y) to get the up vector (Z) to point in the direction of the plane's normal.
 				// Since we are rotating the axis, the rotation is in the opposite direction of the object i.e. -90 degrees.
 				resultPlane.ContentOrientation = UKismetMathLibrary::Conv_VectorToRotator(UKismetMathLibrary::RotateAngleAxis(UKismetMathLibrary::Conv_RotatorToVector(resultPlane.PlaneOrientation), -90, UKismetMathLibrary::GetRightVector(resultPlane.PlaneOrientation)));
+				resultPlane.ID = FGuid(resultMLPlanes[i].id, (resultMLPlanes[i].id >> 32), 0, 0);
 				MLToUnrealPlanesQueryFlags(resultMLPlanes[i].flags, resultPlane.PlaneFlags);
 
 				Planes.Add(resultPlane);
@@ -307,20 +313,9 @@ bool UPlanesComponent::RequestPlanes(int32 UserData, const FPlaneResultDelegate&
 	query.flags = UnrealToMLPlanesQueryFlags(QueryFlags);
 	query.min_hole_length = MinHolePerimeter / WorldToMetersScale;
 	query.min_plane_area = MinPlaneArea / (WorldToMetersScale * WorldToMetersScale);
-	PoseInverse.ConcatenateRotation(SearchVolume->GetComponentQuat());
 	query.bounds_center = MagicLeap::ToMLVector(PoseInverse.TransformPosition(SearchVolume->GetComponentLocation()), WorldToMetersScale);
-	query.bounds_rotation = MagicLeap::ToMLQuat(PoseInverse.GetRotation());
-	query.bounds_extents = MagicLeap::ToMLVector(SearchVolume->GetScaledBoxExtent(), WorldToMetersScale);
-	// MagicLeap::ToMLVector() causes the Z component to be negated.
-	// The bounds were thus invalid and resulted in everything being tracked. 
-	// This provides the content devs with an option to ignore the bounding volume at will.
-	// TODO: Can this be improved?
-	if (!IgnoreBoundingVolume)
-	{
-		query.bounds_extents.x = FMath::Abs<float>(query.bounds_extents.x);
-		query.bounds_extents.y = FMath::Abs<float>(query.bounds_extents.y);
-		query.bounds_extents.z = FMath::Abs<float>(query.bounds_extents.z);
-	}
+	query.bounds_rotation = MagicLeap::ToMLQuat(PoseInverse.TransformRotation(SearchVolume->GetComponentQuat()));
+	query.bounds_extents = MagicLeap::ToMLVectorExtents(SearchVolume->GetScaledBoxExtent(), WorldToMetersScale);
 
 	MLHandle handle;
 	MLResult QueryResult = MLPlanesQueryBegin(Impl->Tracker, &query, &handle);
diff --git a/Engine/Plugins/Lumin/MagicLeap/Source/MagicLeap/Public/IMagicLeapHMD.h b/Engine/Plugins/Lumin/MagicLeap/Source/MagicLeap/Public/IMagicLeapHMD.h
--- a/Engine/Plugins/Lumin/MagicLeap/Source/MagicLeap/Public/IMagicLeapHMD.h
+++ b/Engine/Plugins/Lumin/MagicLeap/Source/MagicLeap/Public/IMagicLeapHMD.h
@@ -7,7 +7,6 @@ class IMagicLeapInputDevice;
 class IMagicLeapHMD
 {
 public:
-	virtual void RegisterMagicLeapInputDevice(IMagicLeapInputDevice* InputDevice) = 0;
-	virtual void UnregisterMagicLeapInputDevice(IMagicLeapInputDevice* InputDevice) = 0;
+	// TODO: no point keeping a single function here.
 	virtual bool IsInitialized() const = 0;
 };
diff --git a/Engine/Plugins/Lumin/MagicLeap/Source/MagicLeap/Public/IMagicLeapInputDevice.h b/Engine/Plugins/Lumin/MagicLeap/Source/MagicLeap/Public/IMagicLeapInputDevice.h
--- a/Engine/Plugins/Lumin/MagicLeap/Source/MagicLeap/Public/IMagicLeapInputDevice.h
+++ b/Engine/Plugins/Lumin/MagicLeap/Source/MagicLeap/Public/IMagicLeapInputDevice.h
@@ -17,4 +17,5 @@ public:
 	virtual void Enable() = 0;
 	virtual bool SupportsExplicitEnable() const = 0;
 	virtual void Disable() = 0;
+	virtual void OnBeginRendering_GameThread_Update() {}
 };
