Updated ML integration to MLSDK 0.19

diff --git a/Engine/Plugins/Lumin/MagicLeap/Source/MagicLeap/Public/IMagicLeapPlugin.h b/Engine/Plugins/Lumin/MagicLeap/Source/MagicLeap/Public/IMagicLeapPlugin.h
--- a/Engine/Plugins/Lumin/MagicLeap/Source/MagicLeap/Public/IMagicLeapPlugin.h
+++ b/Engine/Plugins/Lumin/MagicLeap/Source/MagicLeap/Public/IMagicLeapPlugin.h
@@ -7,6 +7,8 @@
 #include "IHeadMountedDisplayModule.h"
 #include "IMagicLeapHMD.h"
 
+class IMagicLeapInputDevice;
+
 /**
  * The public interface to this module.  In most cases, this interface is only public to sibling modules
  * within this plugin.
@@ -49,4 +51,10 @@ public:
 	* @return Shared pointer to the HMD.
 	*/
 	virtual TWeakPtr<IMagicLeapHMD, ESPMode::ThreadSafe> GetHMD() = 0;
+
+	virtual void RegisterMagicLeapInputDevice(IMagicLeapInputDevice* InputDevice) = 0;
+	virtual void UnregisterMagicLeapInputDevice(IMagicLeapInputDevice* InputDevice) = 0;
+	virtual void EnableInputDevices() = 0;
+	virtual void DisableInputDevices() = 0;
+	virtual void OnBeginRendering_GameThread_UpdateInputDevices() = 0;
 };
diff --git a/Engine/Plugins/Lumin/MagicLeap/Source/MagicLeapController/Private/MagicLeapController.cpp b/Engine/Plugins/Lumin/MagicLeap/Source/MagicLeapController/Private/MagicLeapController.cpp
--- a/Engine/Plugins/Lumin/MagicLeap/Source/MagicLeapController/Private/MagicLeapController.cpp
+++ b/Engine/Plugins/Lumin/MagicLeap/Source/MagicLeapController/Private/MagicLeapController.cpp
@@ -1,7 +1,6 @@
 // Copyright 1998-2018 Epic Games, Inc. All Rights Reserved.
 
 #include "MagicLeapController.h"
-#include "IMagicLeapControllerPlugin.h"
 #include "IMagicLeapPlugin.h"
 #include "MagicLeapHMD.h"
 #include "MagicLeapMath.h"
@@ -15,475 +14,483 @@
 #include "MagicLeapPluginUtil.h"
 #include "TouchpadGesturesComponent.h"
 #include "AssetData.h"
+#if PLATFORM_LUMIN
+#include "Lumin/LuminApplication.h"
+#endif //PLATFORM_LUMIN
 
 #define LOCTEXT_NAMESPACE "MagicLeapController"
 
-//@TODO - would be easier to make sure the Epic enum matches the ML enum if there was a MLInputControllerFeedbackPatternVibe_MAX = MLInputControllerFeedbackPatternVibe_SecondForceDown
-//for now this just guards against insertions
-#if WITH_MLSDK
-static_assert((uint32)EMLControllerHapticPattern::SecondForceDown == (uint32)MLInputControllerFeedbackPatternVibe_SecondForceDown, "Enum entries must match.");
-#endif //WITH_MLSDK
-
+#include "MagicLeapControllerMappings.h"
 
-class FMagicLeapControllerPlugin : public IMagicLeapControllerPlugin
+template<typename T> inline const FString EnumToString(const char* EnumName, T EnumValue)
 {
-public:
+	static const FString InvalidEnum("Invalid");
 
-	virtual void StartupModule() override
+	const UEnum* EnumPtr = FindObject<UEnum>(ANY_PACKAGE, UTF8_TO_TCHAR(EnumName), true);
+	if (!EnumPtr)
 	{
-		IMagicLeapControllerPlugin::StartupModule();
-
-		// HACK: Generic Application might not be instantiated at this point so we create the input device with a 
-		// dummy message handler. When the Generic Application creates the input device it passes a valid message
-		// handler to it which is further on used for all the controller events. This hack fixes issues caused by
-		// using a custom input device before the Generic Application has instantiated it. Eg. within BeginPlay()
-		//
-		// This also fixes the warnings that pop up on the custom input keys when the blueprint loads. Those
-		// warnings are caused because Unreal loads the bluerints before the input device has been instantiated
-		// and has added its keys, thus leading Unreal to believe that those keys don't exist. This hack causes
-		// an earlier instantiation of the input device, and consequently, the custom keys.
-		TSharedPtr<FGenericApplicationMessageHandler> DummyMessageHandler(new FGenericApplicationMessageHandler());
-		CreateInputDevice(DummyMessageHandler.ToSharedRef());
-
-		// Ideally, we should be able to query GetDefault<UMagicLeapSettings>()->bEnableZI directly.
-		// Unfortunately, the UObject system hasn't finished initialization when this module has been loaded.
-		bool bIsVDZIEnabled = false;
-		GConfig->GetBool(TEXT("/Script/MagicLeap.MagicLeapSettings"), TEXT("bEnableZI"), bIsVDZIEnabled, GEngineIni);
-
-		APISetup.Startup(bIsVDZIEnabled);
-#if WITH_MLSDK
-		APISetup.LoadDLL(TEXT("ml_input"));
-#endif
+		return InvalidEnum;
 	}
+	return EnumPtr->GetNameStringByIndex(static_cast<int32>(EnumValue));
+}
 
-	virtual void ShutdownModule() override
-	{
-		APISetup.Shutdown();
-		IMagicLeapControllerPlugin::ShutdownModule();
-	}
+#if WITH_MLSDK
+static_assert(MLInput_MaxControllerTouchpadTouches == FMagicLeapControllerState::kMaxTouches, "Mismatch in max touch constants");
+
+// Controller Mapper
+FMagicLeapController::FControllerMapper::FControllerMapper()
+{
+	// These mappings tell us which entry in the InputControllerState array is providing data for
+	// which motion source.
+	MotionSourceToInputControllerIndex.Add(FMagicLeapMotionSourceNames::Control0, -1);
+	MotionSourceToInputControllerIndex.Add(FMagicLeapMotionSourceNames::Control1, -1);
+	MotionSourceToInputControllerIndex.Add(FMagicLeapMotionSourceNames::MobileApp, -1);
+	InputControllerIndexToMotionSource[0] = FMagicLeapMotionSourceNames::Unknown;
+	InputControllerIndexToMotionSource[1] = FMagicLeapMotionSourceNames::Unknown;
+
+	DefaultInputControllerIndexToHand[0] = EControllerHand::Right;
+	DefaultInputControllerIndexToHand[1] = EControllerHand::Left;
+}
 
-	virtual TSharedPtr<class IInputDevice> CreateInputDevice(const TSharedRef<FGenericApplicationMessageHandler>& InMessageHandler) override
+void FMagicLeapController::FControllerMapper::UpdateMotionSourceInputIndexPairing(const MLInputControllerState ControllerState[MLInput_MaxControllers])
+{
+	// Determine which entry in the ControllerState array is providing data for which motion source.
+	// This is kind of messy, and is the result of an old mandate to allow MLMA to substitute for a control.
+	// Once the platform API is adjusted to treat these as separate devices this will be much cleaner.
+	TMap<FName, int32> NewMotionSourceToInputControllerIndex;
+	FName NewInputControllerIndexToMotionSource[MLInput_MaxControllers];
+
+	NewMotionSourceToInputControllerIndex.Add(FMagicLeapMotionSourceNames::Control0, -1);
+	NewMotionSourceToInputControllerIndex.Add(FMagicLeapMotionSourceNames::Control1, -1);
+	NewMotionSourceToInputControllerIndex.Add(FMagicLeapMotionSourceNames::MobileApp, -1);
+	NewInputControllerIndexToMotionSource[0] = FMagicLeapMotionSourceNames::Unknown;
+	NewInputControllerIndexToMotionSource[1] = FMagicLeapMotionSourceNames::Unknown;
+
+	for (int i = 0; i < MLInput_MaxControllers; ++i)
 	{
-		if (!InputDevice.IsValid())
+		switch (ControllerState[i].type)
 		{
-			TSharedPtr<IInputDevice> MagicLeapController(new FMagicLeapController(InMessageHandler));
-			InputDevice = MagicLeapController;
-			return InputDevice;
-		}
-		else
+		case MLInputControllerType_MobileApp:
 		{
-			InputDevice.Get()->SetMessageHandler(InMessageHandler);
-			return InputDevice;
+			NewMotionSourceToInputControllerIndex[FMagicLeapMotionSourceNames::MobileApp] = i;
+			NewInputControllerIndexToMotionSource[i] = FMagicLeapMotionSourceNames::MobileApp;
+			break;
 		}
-		return nullptr;
-	}
-
-	virtual TSharedPtr<IInputDevice> GetInputDevice() override
-	{
-		if (!InputDevice.IsValid())
+		case MLInputControllerType_Device:
 		{
-			InputDevice = CreateInputDevice(FSlateApplication::Get().GetPlatformApplication()->GetMessageHandler());
+			if (ControllerState[i].hardware_index == 0)
+			{
+				NewMotionSourceToInputControllerIndex[FMagicLeapMotionSourceNames::Control0] = i;
+				NewInputControllerIndexToMotionSource[i] = FMagicLeapMotionSourceNames::Control0;
+			}
+			else
+			{
+				NewMotionSourceToInputControllerIndex[FMagicLeapMotionSourceNames::Control1] = i;
+				NewInputControllerIndexToMotionSource[i] = FMagicLeapMotionSourceNames::Control1;
+			}
+			break;
+		}
+		default:
+			break;
 		}
-		return InputDevice;
 	}
 
-private:
-	FMagicLeapAPISetup APISetup;
-	TSharedPtr<IInputDevice> InputDevice;
-};
-
-IMPLEMENT_MODULE(FMagicLeapControllerPlugin, MagicLeapController);
-
-//////////////////////////////////////////////////////////////////////////
-
-const FKey FMagicLeapController::FMagicLeapKeys::MotionController_Left_Thumbstick_Z("MotionController_Left_Thumbstick_Z");
-const FName FMagicLeapControllerKeyNames::MotionController_Left_Thumbstick_Z_Name("MotionController_Left_Thumbstick_Z");
-
-const FKey FMagicLeapController::FMagicLeapKeys::Left_MoveButton("MagicLeap_Left_MoveButton");
-const FName FMagicLeapControllerKeyNames::Left_MoveButton_Name("MagicLeap_Left_MoveButton");
-
-const FKey FMagicLeapController::FMagicLeapKeys::Left_AppButton("MagicLeap_Left_AppButton");
-const FName FMagicLeapControllerKeyNames::Left_AppButton_Name("MagicLeap_Left_AppButton");
-
-const FKey FMagicLeapController::FMagicLeapKeys::Left_HomeButton("MagicLeap_Left_HomeButton");
-const FName FMagicLeapControllerKeyNames::Left_HomeButton_Name("MagicLeap_Left_HomeButton");
-
-const FKey FMagicLeapController::FMagicLeapKeys::MotionController_Right_Thumbstick_Z("MotionController_Right_Thumbstick_Z");
-const FName FMagicLeapControllerKeyNames::MotionController_Right_Thumbstick_Z_Name("MotionController_Right_Thumbstick_Z");
-
-const FKey FMagicLeapController::FMagicLeapKeys::Right_MoveButton("MagicLeap_Right_MoveButton");
-const FName FMagicLeapControllerKeyNames::Right_MoveButton_Name("MagicLeap_Right_MoveButton");
-
-const FKey FMagicLeapController::FMagicLeapKeys::Right_AppButton("MagicLeap_Right_AppButton");
-const FName FMagicLeapControllerKeyNames::Right_AppButton_Name("MagicLeap_Right_AppButton");
-
-const FKey FMagicLeapController::FMagicLeapKeys::Right_HomeButton("MagicLeap_Right_HomeButton");
-const FName FMagicLeapControllerKeyNames::Right_HomeButton_Name("MagicLeap_Right_HomeButton");
-
-FMagicLeapController::FMagicLeapController(const TSharedRef<FGenericApplicationMessageHandler>& InMessageHandler)
-	: MessageHandler(InMessageHandler)
-	, DeviceIndex(0)  // Input Controller Index for Unreal is hardcoded to 0. Ideally it should be incremented for each registered InputDevice.
-#if WITH_MLSDK
-	, InputTracker(ML_INVALID_HANDLE)
-	, ControllerTracker(ML_INVALID_HANDLE)
-#endif //WITH_MLSDK
-	, bIsInputStateValid(false)
-	, bTriggerState(false)
-	, bTriggerKeyPressing(false)
-	, triggerKeyIsConsideredPressed(80.0f)
-	, triggerKeyIsConsideredReleased(20.0f)
-{
-	InitializeInputCallbacks();
-
-	//hack call to empty function to force the automation tests to compile in
-	MagicLeapTestReferenceFunction();
-
-	int32 ControllerIndex = 0;
-	// Make Right hand the default hand (Controller index 0), for backwards compatibility with the previous versions.
-	ControllerIDToHand.Add(ControllerIndex++, EControllerHand::Right);
-	ControllerIDToHand.Add(ControllerIndex++, EControllerHand::Left);
-#if WITH_MLSDK
-	check(ControllerIndex <= MLInput_MaxControllers);
-#endif //WITH_MLSDK
-
-	// Make Right hand the default hand (Controller index 0), for backwards compatibility with the previous versions.
-	ControllerIndex = 0;
-	HandToControllerID.Add(EControllerHand::Right, ControllerIndex++);
-	HandToControllerID.Add(EControllerHand::Left, ControllerIndex++);
-#if WITH_MLSDK
-	check(ControllerIndex <= MLInput_MaxControllers);
-#endif //WITH_MLSDK
-
-	// Shadow state for the controllers
-	ControllerIndex = 0;
-	ControllerIDToControllerState.Add(ControllerIndex++, FMagicLeapControllerState(EControllerHand::Right));
-	ControllerIDToControllerState.Add(ControllerIndex++, FMagicLeapControllerState(EControllerHand::Left));
-#if WITH_MLSDK
-	check(ControllerIndex <= MLInput_MaxControllers);
-#endif //WITH_MLSDK
-
-	if (GConfig)
+	// Only do the guarded copy if anything changed
+	if (FMemory::Memcmp(InputControllerIndexToMotionSource,
+		NewInputControllerIndexToMotionSource, sizeof(NewInputControllerIndexToMotionSource)) != 0)
 	{
-		float floatValueReceived = 0.0f;
-		GConfig->GetFloat(TEXT("/Script/Engine.InputSettings"), TEXT("TriggerKeyIsConsideredPressed"), floatValueReceived, GInputIni);
-		triggerKeyIsConsideredPressed = floatValueReceived;
-
-		GConfig->GetFloat(TEXT("/Script/Engine.InputSettings"), TEXT("TriggerKeyIsConsideredReleased"), floatValueReceived, GInputIni);
-		triggerKeyIsConsideredReleased = floatValueReceived;
-	}
-
-	// Register "MotionController" modular feature manually
-	IModularFeatures::Get().RegisterModularFeature(GetModularFeatureName(), this);
-	AddKeys();
+		FScopeLock Lock(&CriticalSection);
 
-	// We're implicitly requiring that the MagicLeapPlugin has been loaded and
-	// initialized at this point.
-	auto HMD = IMagicLeapPlugin::Get().GetHMD().Pin();
-	if (HMD.IsValid())
-	{
-		HMD->RegisterMagicLeapInputDevice(this);
+		FMemory::Memcpy(InputControllerIndexToMotionSource,
+			NewInputControllerIndexToMotionSource, sizeof(NewInputControllerIndexToMotionSource));
+		MotionSourceToInputControllerIndex = NewMotionSourceToInputControllerIndex;
 	}
 }
 
-FMagicLeapController::~FMagicLeapController()
+void FMagicLeapController::FControllerMapper::MapHandToMotionSource(EControllerHand Hand, FName MotionSource)
 {
-	// Normally, the MagicLeapPlugin will be around during unload,
-	// but it isn't an assumption that we should make.
-	if (IMagicLeapPlugin::IsAvailable())
+	if (Hand == EControllerHand::Right || Hand == EControllerHand::Left)
 	{
-		auto HMD = IMagicLeapPlugin::Get().GetHMD().Pin();
-		if (HMD.IsValid())
+		FScopeLock Lock(&CriticalSection);
+
+		if (MotionSource == FMagicLeapMotionSourceNames::Control0 ||
+			MotionSource == FMagicLeapMotionSourceNames::Control1 ||
+			MotionSource == FMagicLeapMotionSourceNames::MobileApp)
+		{
+			// Make sure to not allow multiple motion sources to point to the same hand
+			auto ExistingMapping = HandToMotionSource.Find(Hand);
+			if (ExistingMapping != nullptr)
+			{
+				*ExistingMapping = MotionSource;
+			}
+			else
+			{
+				HandToMotionSource.Add(Hand, MotionSource);
+			}
+			MotionSourceToHand.FindOrAdd(MotionSource) = Hand;
+		}
+		else
 		{
-			HMD->UnregisterMagicLeapInputDevice(this);
+			// Our module will not map non-ML devices
+			auto ExistingMapping = HandToMotionSource.Find(Hand);
+			if (ExistingMapping != nullptr)
+			{
+				MotionSourceToHand.Remove(*ExistingMapping);
+				HandToMotionSource.Remove(Hand);
+			}
 		}
 	}
-
-	Disable();
-
-	IModularFeatures::Get().UnregisterModularFeature(GetModularFeatureName(), this);
 }
 
-void FMagicLeapController::OnChar(uint32 CharUTF32)
+FName FMagicLeapController::FControllerMapper::GetMotionSourceForHand(EControllerHand Hand) const
 {
-	FMagicLeapHMD::EnableInput EnableInputFromHMD;
-	// fixes unreferenced parameter error for Windows package builds.
-	(void)EnableInputFromHMD;
-	MessageHandler->OnKeyChar(CharUTF32, false);
-}
+	FScopeLock Lock(const_cast<FCriticalSection*>(&CriticalSection));
 
-#if WITH_MLSDK
-void FMagicLeapController::OnKeyDown(MLKeyCode KeyCode)
-{
-	FKey pressedKey = FInputKeyManager::Get().GetKeyFromCodes(static_cast<uint32>(KeyCode), static_cast<uint32>(KeyCode));
-	if (pressedKey != EKeys::Invalid)
+	if (HandToMotionSource.Num() == 0)
 	{
-		FMagicLeapHMD::EnableInput EnableInputFromHMD;
-		// fixes unreferenced parameter error for Windows package builds.
-		(void)EnableInputFromHMD;
-		MessageHandler->OnControllerButtonPressed(pressedKey.GetFName(), DeviceIndex, false);
+		if (Hand == DefaultInputControllerIndexToHand[0])
+			return InputControllerIndexToMotionSource[0];
+		return InputControllerIndexToMotionSource[1];
 	}
 	else
 	{
-		UE_LOG(LogMagicLeapController, Error, TEXT("Key not supported in Unreal. Keycode = %d"), static_cast<uint32>(KeyCode));
+		const auto& MotionSource = HandToMotionSource.Find(Hand);
+		if (MotionSource != nullptr)
+		{
+			return *MotionSource;
+		}
 	}
+	return FMagicLeapMotionSourceNames::Unknown;
 }
 
-void FMagicLeapController::OnKeyUp(MLKeyCode KeyCode)
+EControllerHand FMagicLeapController::FControllerMapper::GetHandForMotionSource(FName MotionSource) const
 {
-	FKey pressedKey = FInputKeyManager::Get().GetKeyFromCodes(static_cast<uint32>(KeyCode), static_cast<uint32>(KeyCode));
-	if (pressedKey != EKeys::Invalid)
+	auto ControllerHand = EControllerHand::ControllerHand_Count;
+
+	// Legacy hand motion sources
+	if (FXRMotionControllerBase::GetHandEnumForSourceName(MotionSource, ControllerHand))
 	{
-		FMagicLeapHMD::EnableInput EnableInputFromHMD;
-		// fixes unreferenced parameter error for Windows package builds.
-		(void)EnableInputFromHMD;
-		MessageHandler->OnControllerButtonReleased(pressedKey.GetFName(), DeviceIndex, false);
+		// Only left and right are allowed
+		if (ControllerHand != EControllerHand::Right && ControllerHand != EControllerHand::Left)
+		{
+			ControllerHand = EControllerHand::ControllerHand_Count;
+		}
 	}
 	else
 	{
-		UE_LOG(LogMagicLeapController, Error, TEXT("Key not supported in Unreal. Keycode = %d"), static_cast<uint32>(KeyCode));
+		FScopeLock Lock(const_cast<FCriticalSection*>(&CriticalSection));
+
+		if (HandToMotionSource.Num() == 0)
+		{
+			if (InputControllerIndexToMotionSource[0] == MotionSource)
+			{
+				ControllerHand = EControllerHand::Right;
+			}
+			else if (InputControllerIndexToMotionSource[1] == MotionSource)
+			{
+				ControllerHand = EControllerHand::Left;
+			}
+		}
+		else
+		{
+			const auto& Hand = MotionSourceToHand.Find(MotionSource);
+			if (Hand != nullptr)
+			{
+				ControllerHand = *Hand;
+			}
+		}
 	}
+	return ControllerHand;
 }
 
-#define TOUCH_GESTURE_CASE(x) case MLInputControllerTouchpadGestureType_##x: { return EMagicLeapTouchpadGestureType::x; }
-
-EMagicLeapTouchpadGestureType MLToUnrealTouchpadGestureType(MLInputControllerTouchpadGestureType GestureType)
+FName FMagicLeapController::FControllerMapper::GetMotionSourceForInputControllerIndex(uint8 controller_id) const
 {
-	switch (GestureType)
+	if (controller_id < MLInput_MaxControllers)
 	{
-		TOUCH_GESTURE_CASE(Tap)
-		TOUCH_GESTURE_CASE(ForceTapDown)
-		TOUCH_GESTURE_CASE(ForceTapUp)
-		TOUCH_GESTURE_CASE(ForceDwell)
-		TOUCH_GESTURE_CASE(SecondForceDown)
-		TOUCH_GESTURE_CASE(LongHold)
-		TOUCH_GESTURE_CASE(RadialScroll)
-		TOUCH_GESTURE_CASE(Swipe)
-		TOUCH_GESTURE_CASE(Scroll)
-		TOUCH_GESTURE_CASE(Pinch)
+		FScopeLock Lock(const_cast<FCriticalSection*>(&CriticalSection));
+		return InputControllerIndexToMotionSource[controller_id];
 	}
-	return EMagicLeapTouchpadGestureType::None;
+	return FMagicLeapMotionSourceNames::Unknown;
 }
 
-#define TOUCH_DIR_CASE(x) case MLInputControllerTouchpadGestureDirection_##x: { return EMagicLeapTouchpadGestureDirection::x; }
-
-EMagicLeapTouchpadGestureDirection MLToUnrealTouchpadGestureDirection(MLInputControllerTouchpadGestureDirection Direction)
+uint8 FMagicLeapController::FControllerMapper::GetInputControllerIndexForMotionSource(FName MotionSource) const
 {
-	switch (Direction)
+	auto ControllerHand = EControllerHand::ControllerHand_Count;
+
+	// Legacy hand motion sources
+	if (FXRMotionControllerBase::GetHandEnumForSourceName(MotionSource, ControllerHand))
+	{
+		if (ControllerHand == EControllerHand::Right || ControllerHand == EControllerHand::Left)
+		{
+			return GetInputControllerIndexForHand(ControllerHand);
+		}
+	}
+	else
 	{
-		TOUCH_DIR_CASE(Up)
-		TOUCH_DIR_CASE(Down)
-		TOUCH_DIR_CASE(Left)
-		TOUCH_DIR_CASE(Right)
-		TOUCH_DIR_CASE(In)
-		TOUCH_DIR_CASE(Out)
-		TOUCH_DIR_CASE(Clockwise)
-		TOUCH_DIR_CASE(CounterClockwise)
+		FScopeLock Lock(const_cast<FCriticalSection*>(&CriticalSection));
+		if (InputControllerIndexToMotionSource[0] == MotionSource)
+			return 0;
+		if (InputControllerIndexToMotionSource[1] == MotionSource)
+			return 1;
 	}
-	return EMagicLeapTouchpadGestureDirection::None;
+	return 0xFF;
 }
 
-FMagicLeapTouchpadGesture MLToUnrealTouchpadGesture(EControllerHand hand, const MLInputControllerTouchpadGesture& touchpad_gesture)
+EControllerHand FMagicLeapController::FControllerMapper::GetHandForInputControllerIndex(uint8 controller_id) const
 {
-	FMagicLeapTouchpadGesture gesture;
-	gesture.Hand = hand;
-	gesture.Type = MLToUnrealTouchpadGestureType(touchpad_gesture.type);
-	gesture.Direction = MLToUnrealTouchpadGestureDirection(touchpad_gesture.direction);
-	gesture.PositionAndForce = FVector(touchpad_gesture.pos_and_force.x, touchpad_gesture.pos_and_force.y, touchpad_gesture.pos_and_force.z);
-	gesture.Speed = touchpad_gesture.speed;
-	gesture.Distance = touchpad_gesture.distance;
-	gesture.FingerGap = touchpad_gesture.finger_gap;
-	gesture.Radius = touchpad_gesture.radius;
-	gesture.Angle = touchpad_gesture.angle;
-
-	return gesture;
-}
+	if (controller_id < MLInput_MaxControllers)
+	{
+		if (HandToMotionSource.Num() == 0)
+		{
+			return DefaultInputControllerIndexToHand[controller_id];
+		}
+		else
+		{
+			FScopeLock Lock(const_cast<FCriticalSection*>(&CriticalSection));
 
-void FMagicLeapController::OnButtonDown(uint8 ControllerID, MLInputControllerButton Button)
-{
-	FMagicLeapHMD::EnableInput EnableInputFromHMD;
-	// fixes unreferenced parameter error for Windows package builds.
-	(void)EnableInputFromHMD;
-	MessageHandler->OnControllerButtonPressed(MagicLeapButtonToUnrealButton(ControllerID, Button), DeviceIndex, false);
-}
+			const auto& Hand = MotionSourceToHand.Find(InputControllerIndexToMotionSource[controller_id]);
+			if (Hand != nullptr)
+			{
+				return *Hand;
+			}
 
-void FMagicLeapController::OnButtonUp(uint8 ControllerID, MLInputControllerButton Button)
-{
-	FMagicLeapHMD::EnableInput EnableInputFromHMD;
-	// fixes unreferenced parameter error for Windows package builds.
-	(void)EnableInputFromHMD;
-	MessageHandler->OnControllerButtonReleased(MagicLeapButtonToUnrealButton(ControllerID, Button), DeviceIndex, false);
+		}
+	}
+	return EControllerHand::ControllerHand_Count;
 }
-#endif //WITH_MLSDK
 
-void FMagicLeapController::OnControllerConnected(EControllerHand Hand)
+uint8 FMagicLeapController::FControllerMapper::GetInputControllerIndexForHand(EControllerHand Hand) const
 {
-	UE_LOG(LogCore, Log, TEXT("CONTROLLER CONNECTED"));
-	SetControllerIsConnected(Hand, true);
-}
+	if (HandToMotionSource.Num() == 0)
+	{
+		if (Hand == DefaultInputControllerIndexToHand[0])
+		{
+			return 0;
+		}
+		return 1;
+	}
+	else
+	{
+		FScopeLock Lock(const_cast<FCriticalSection*>(&CriticalSection));
 
-void FMagicLeapController::OnControllerDisconnected(EControllerHand Hand)
-{
-	UE_LOG(LogCore, Log, TEXT("CONTROLLER DISCONNECTED"));
-	SetControllerIsConnected(Hand, false);
+		const auto& MotionSource = HandToMotionSource.Find(Hand);
+		if (MotionSource != nullptr)
+		{
+			return MotionSourceToInputControllerIndex[*MotionSource];
+		}
+	}
+	return 0xFF;
 }
 
-void FMagicLeapController::SetControllerIsConnected(EControllerHand Hand, bool bConnected)
+EMLControllerType FMagicLeapController::FControllerMapper::MotionSourceToControllerType(FName MotionSource)
 {
-#if WITH_MLSDK
-	const int32* ControllerIndex = HandToControllerID.Find((EControllerHand)Hand);
-	if (ControllerIndex != nullptr)
+	auto MLSourceToControllerType = [](FName InMotionSource)
 	{
-		//get the shadow state from the controller id
-		FMagicLeapControllerState* ControllerState = ControllerIDToControllerState.Find(*ControllerIndex);
-		if (ControllerState != nullptr)
+		if (InMotionSource == FMagicLeapMotionSourceNames::Control0 ||
+			InMotionSource == FMagicLeapMotionSourceNames::Control1)
 		{
-			ControllerState->bIsConnected = bConnected;
-
-			//default to no tracking
-			ControllerState->bIsOrientationTracked = false;
-			ControllerState->bIsPositionTracked = false;
+			return EMLControllerType::Device;
+		}
+		if (InMotionSource == FMagicLeapMotionSourceNames::MobileApp)
+		{
+			return EMLControllerType::MobileApp;
+		}
+		return EMLControllerType::None;
+	};
 
-			switch (InputState[*ControllerIndex].dof)
+	// First just see if it's one of ours and can be easily mapped
+	auto ControllerType = MLSourceToControllerType(MotionSource);
+	if (ControllerType == EMLControllerType::None)
+	{
+		// If not, see if it's a hand mapping
+		auto ControllerHand = EControllerHand::ControllerHand_Count;
+		if (FXRMotionControllerBase::GetHandEnumForSourceName(MotionSource, ControllerHand) &&
+			(ControllerHand == EControllerHand::Right || ControllerHand == EControllerHand::Left))
+		{
+			TSharedPtr<FMagicLeapController> controller =
+				StaticCastSharedPtr<FMagicLeapController>(IMagicLeapControllerPlugin::Get().GetInputDevice());
+			if (controller.IsValid())
 			{
-			case MLInputControllerDof_3:
-				ControllerState->bIsOrientationTracked = true;
-				break;
-			case MLInputControllerDof_6:
-				ControllerState->bIsOrientationTracked = true;
-				ControllerState->bIsPositionTracked = true;
-				break;
-			default:
-				break;
+				auto TheMotionSource = controller->
+					ControllerMapper.GetMotionSourceForHand(ControllerHand);
+				ControllerType = MotionSourceToControllerType(TheMotionSource);
 			}
 		}
 	}
-#endif //WITH_MLSDK
+	return ControllerType;
 }
 
-void FMagicLeapController::Tick(float DeltaTime)
+void FMagicLeapController::FControllerMapper::SwapHands()
 {
-	UpdateTrackerData();
+	Swap(DefaultInputControllerIndexToHand[0], DefaultInputControllerIndexToHand[1]);
 }
+#endif //WITH_MLSDK
 
-void FMagicLeapController::SendControllerEvents()
+FMagicLeapController::FMagicLeapController(const TSharedRef<FGenericApplicationMessageHandler>& InMessageHandler)
+	: MessageHandler(InMessageHandler)
+	, DeviceIndex(0)  // Input Controller Index for Unreal is hardcoded to 0. Ideally it should be incremented for each registered InputDevice.
+#if WITH_MLSDK
+	, InputTracker(ML_INVALID_HANDLE)
+	, ControllerTracker(ML_INVALID_HANDLE)
+	, ControllerDof(MLInputControllerDof_6)
+	, TrackingMode(EMLControllerTrackingMode::CoordinateFrameUID)
+#endif //WITH_MLSDK
+	, bIsInputStateValid(false)
+	, TriggerKeyIsConsideredPressed(80.0f)
+	, TriggerKeyIsConsideredReleased(20.0f)
 {
+	// Hack call to empty function to force the automation tests to compile in
+	MagicLeapTestReferenceFunction();
+
 #if WITH_MLSDK
-	if (bIsInputStateValid && MessageHandler.IsValid())
+	InitializeInputCallbacks();
+
+	// Current and previous frame of engine-mapped controller data
+	CurrMotionSourceControllerState.Add(FMagicLeapMotionSourceNames::Control0);
+	CurrMotionSourceControllerState.Add(FMagicLeapMotionSourceNames::Control1);
+	CurrMotionSourceControllerState.Add(FMagicLeapMotionSourceNames::MobileApp);
+	PrevMotionSourceControllerState.Add(FMagicLeapMotionSourceNames::Control0);
+	PrevMotionSourceControllerState.Add(FMagicLeapMotionSourceNames::Control1);
+	PrevMotionSourceControllerState.Add(FMagicLeapMotionSourceNames::MobileApp);
+#endif //WITH_MLSDK
+
+	// Register "MotionController" modular feature manually
+	IModularFeatures::Get().RegisterModularFeature(GetModularFeatureName(), this);
+	AddKeys();
+
+	// We're implicitly requiring that the MagicLeapPlugin has been loaded and
+	// initialized at this point.
+	IMagicLeapPlugin::Get().RegisterMagicLeapInputDevice(this);
+}
+
+FMagicLeapController::~FMagicLeapController()
+{
+	// Normally, the MagicLeapPlugin will be around during unload,
+	// but it isn't an assumption that we should make.
+	if (IMagicLeapPlugin::IsAvailable())
 	{
-		SendControllerEventsForHand(EControllerHand::Left);
-		SendControllerEventsForHand(EControllerHand::Right);
+		IMagicLeapPlugin::Get().UnregisterMagicLeapInputDevice(this);
+	}
+
+	Disable();
+
+	IModularFeatures::Get().UnregisterModularFeature(GetModularFeatureName(), this);
+}
+
+void FMagicLeapController::InitializeInputCallbacks()
+{
+#if WITH_MLSDK
+	FMemory::Memset(&InputControllerCallbacks, 0, sizeof(InputControllerCallbacks));
 
+	InputControllerCallbacks.on_touchpad_gesture_start = [](uint8 controller_id, const MLInputControllerTouchpadGesture *touchpad_gesture, void *data)
+	{
+		auto Controller = reinterpret_cast<FMagicLeapController*>(data);
+		if (Controller && controller_id < MLInput_MaxControllers)
 		{
-			FScopeLock Lock(&ButtonCriticalSection);
-			for (const FButtonMap& buttonmap : PendingButtonStates)
+			const auto& Hand = Controller->ControllerMapper.GetHandForInputControllerIndex(controller_id);
+			if (Hand != EControllerHand::ControllerHand_Count)
 			{
-				if (buttonmap.bPressed)
+				auto Gesture = MLToUnrealTouchpadGesture(Hand, FMagicLeapMotionSourceNames::Unknown, *touchpad_gesture);
+				for (IMagicLeapTouchpadGestures* Receiver : Controller->TouchpadGestureReceivers)
 				{
-					OnButtonDown(buttonmap.ControllerID, buttonmap.Button);
+					Receiver->OnTouchpadGestureStartCallback(Gesture);
 				}
-				else
+			}
+		}
+	};
+
+	InputControllerCallbacks.on_touchpad_gesture_continue = [](uint8 controller_id, const MLInputControllerTouchpadGesture *touchpad_gesture, void *data)
+	{
+		auto Controller = reinterpret_cast<FMagicLeapController*>(data);
+		if (Controller && controller_id < MLInput_MaxControllers)
+		{
+			const auto& Hand = Controller->ControllerMapper.GetHandForInputControllerIndex(controller_id);
+			if (Hand != EControllerHand::ControllerHand_Count)
+			{
+				auto Gesture = MLToUnrealTouchpadGesture(Hand, FMagicLeapMotionSourceNames::Unknown, *touchpad_gesture);
+				for (IMagicLeapTouchpadGestures* Receiver : Controller->TouchpadGestureReceivers)
 				{
-					OnButtonUp(buttonmap.ControllerID, buttonmap.Button);
+					Receiver->OnTouchpadGestureContinueCallback(Gesture);
 				}
 			}
-			PendingButtonStates.Empty();
 		}
+	};
 
+	InputControllerCallbacks.on_touchpad_gesture_end = [](uint8 controller_id, const MLInputControllerTouchpadGesture *touchpad_gesture, void *data)
+	{
+		auto Controller = reinterpret_cast<FMagicLeapController*>(data);
+		if (Controller && controller_id < MLInput_MaxControllers)
 		{
-			FScopeLock Lock(&KeyCriticalSection);
-			for (MLKeyCode keyCode : PendingKeyDowns)
+			const auto& Hand = Controller->ControllerMapper.GetHandForInputControllerIndex(controller_id);
+			if (Hand != EControllerHand::ControllerHand_Count)
 			{
-				OnKeyDown(keyCode);
+				auto Gesture = MLToUnrealTouchpadGesture(Hand, FMagicLeapMotionSourceNames::Unknown, *touchpad_gesture);
+				for (IMagicLeapTouchpadGestures* Receiver : Controller->TouchpadGestureReceivers)
+				{
+					Receiver->OnTouchpadGestureEndCallback(Gesture);
+				}
 			}
-			for (MLKeyCode keyCode : PendingKeyUps)
+		}
+	};
+
+	InputControllerCallbacks.on_button_down = [](uint8_t controller_id, MLInputControllerButton button, void *data)
+	{
+		auto Controller = reinterpret_cast<FMagicLeapController*>(data);
+		if (Controller && controller_id < MLInput_MaxControllers)
+		{
+			const auto ControllerHand = Controller->ControllerMapper.GetHandForInputControllerIndex(controller_id);
+			if (ControllerHand != EControllerHand::ControllerHand_Count)
 			{
-				OnKeyUp(keyCode);
+				Controller->PendingButtonEvents.Enqueue(MakeTuple(MLToUnrealButton(ControllerHand, button), true));
 			}
-			for (uint32 charUTF32 : PendingCharKeys)
+		}
+	};
+
+	InputControllerCallbacks.on_button_up = [](uint8_t controller_id, MLInputControllerButton button, void *data)
+	{
+		auto Controller = reinterpret_cast<FMagicLeapController*>(data);
+		if (Controller && controller_id < MLInput_MaxControllers)
+		{
+			const auto ControllerHand = Controller->ControllerMapper.GetHandForInputControllerIndex(controller_id);
+			if (ControllerHand != EControllerHand::ControllerHand_Count)
 			{
-				OnChar(charUTF32);
+				Controller->PendingButtonEvents.Enqueue(MakeTuple(MLToUnrealButton(ControllerHand, button), false));
 			}
-
-			PendingKeyDowns.Empty();
-			PendingKeyUps.Empty();
-			PendingCharKeys.Empty();
 		}
-	}
+	};
 #endif //WITH_MLSDK
 }
 
-void FMagicLeapController::SendControllerEventsForHand(EControllerHand Hand)
+void FMagicLeapController::Tick(float DeltaTime)
+{
+	UpdateTrackerData();
+}
+
+void FMagicLeapController::SendControllerEvents()
 {
 #if WITH_MLSDK
-	const int32* Controller = HandToControllerID.Find(Hand);
-	if (Controller != nullptr)
+	if (bIsInputStateValid && MessageHandler.IsValid())
 	{
-		const int32 ControllerID = *Controller;
+		FMagicLeapHMD::EnableInput EnableInputFromHMD;
+		// fixes unreferenced parameter error for Windows package builds.
+		(void)EnableInputFromHMD;
 
-		// Controller connected or disconnected
-		if (InputState[ControllerID].is_connected && !OldInputState[ControllerID].is_connected)
-		{
-			OnControllerConnected(Hand);
-		}
-		else if (!InputState[ControllerID].is_connected && OldInputState[ControllerID].is_connected)
-		{
-			OnControllerDisconnected(Hand);
-		}
+		SendControllerEventsForHand(EControllerHand::Right);
+		SendControllerEventsForHand(EControllerHand::Left);
 
-		if (InputState[ControllerID].is_connected)
+		while (!PendingButtonEvents.IsEmpty())
 		{
-			// Touch 0 maps to Motion Controller Thumbstick
-			if (InputState[ControllerID].is_touch_active[0])
+			TPair<FName, bool> ButtonEvent;
+
+			PendingButtonEvents.Dequeue(ButtonEvent);
+
+			if (ButtonEvent.Value)
 			{
-				FMagicLeapHMD::EnableInput EnableInputFromHMD;
-				// fixes unreferenced parameter error for Windows package builds.
-				(void)EnableInputFromHMD;
-				MessageHandler->OnControllerAnalog(MagicLeapTouchToUnrealThumbstickAxis(Hand, 0),
-					DeviceIndex, InputState[ControllerID].touch_pos_and_force[0].x);
-				MessageHandler->OnControllerAnalog(MagicLeapTouchToUnrealThumbstickAxis(Hand, 1),
-					DeviceIndex, InputState[ControllerID].touch_pos_and_force[0].y);
-				MessageHandler->OnControllerAnalog(MagicLeapTouchToUnrealThumbstickAxis(Hand, 2),
-					DeviceIndex, InputState[ControllerID].touch_pos_and_force[0].z);
+				MessageHandler->OnControllerButtonPressed(ButtonEvent.Key, DeviceIndex, false);
 			}
 			else
 			{
-				FMagicLeapHMD::EnableInput EnableInputFromHMD;
-				// fixes unreferenced parameter error for Windows package builds.
-				(void)EnableInputFromHMD;
-				MessageHandler->OnControllerAnalog(MagicLeapTouchToUnrealThumbstickAxis(Hand, 0), DeviceIndex, 0.0f);
-				MessageHandler->OnControllerAnalog(MagicLeapTouchToUnrealThumbstickAxis(Hand, 1), DeviceIndex, 0.0f);
-				MessageHandler->OnControllerAnalog(MagicLeapTouchToUnrealThumbstickAxis(Hand, 2), DeviceIndex, 0.0f);
-			}
-
-			DebouncedButtonMessageHandler(InputState[ControllerID].is_touch_active[0],
-				OldInputState[ControllerID].is_touch_active[0], MagicLeapTouchToUnrealThumbstickButton(Hand));
-
-			// Analog trigger
-			if (OldInputState[ControllerID].trigger_normalized != InputState[ControllerID].trigger_normalized)
-			{
-				FMagicLeapHMD::EnableInput EnableInputFromHMD;
-				// fixes unreferenced parameter error for Windows package builds.
-				(void)EnableInputFromHMD;
-				MessageHandler->OnControllerAnalog(MagicLeapTriggerToUnrealTriggerAxis(Hand),
-					DeviceIndex, InputState[ControllerID].trigger_normalized);
-
-				// TODO: This is a temporary hack until the C-API can send trigger down messages
-				const float triggerValue = InputState[ControllerID].trigger_normalized;
-				const bool isTriggerKeyPressing = (triggerValue > triggerKeyIsConsideredPressed) && !bTriggerKeyPressing;
-				const bool isTriggerKeyReleasing = (triggerValue < triggerKeyIsConsideredReleased) && bTriggerKeyPressing;
-
-				if (isTriggerKeyPressing)
-				{
-					// Analog trigger pressed
-					MessageHandler->OnControllerButtonPressed(MagicLeapTriggerToUnrealTriggerKey(Hand), DeviceIndex, false);
-					bTriggerKeyPressing = true;
-				}
-				else if (isTriggerKeyReleasing)
-				{
-					// Analog trigger released
-					MessageHandler->OnControllerButtonReleased(MagicLeapTriggerToUnrealTriggerKey(Hand), DeviceIndex, false);
-					bTriggerKeyPressing = false;
-				}
+				MessageHandler->OnControllerButtonReleased(ButtonEvent.Key, DeviceIndex, false);
 			}
 		}
 	}
@@ -500,138 +507,177 @@ bool FMagicLeapController::Exec(UWorld* InWorld, const TCHAR* Cmd, FOutputDevice
 	return false;
 }
 
-void FMagicLeapController::SetChannelValue(int32 ControllerId, FForceFeedbackChannelType ChannelType, float Value)
+bool FMagicLeapController::IsGamepadAttached() const
 {
-	const EControllerHand Hand = (ChannelType == FForceFeedbackChannelType::LEFT_LARGE || ChannelType == FForceFeedbackChannelType::LEFT_SMALL) ? EControllerHand::Left : EControllerHand::Right;
-
-	bool bFromHapticInterface = false;
-	FHapticFeedbackValues HapticValues;
-	HapticValues.Frequency = 1.0f;
-	HapticValues.Amplitude = Value;
-	InternalSetHapticFeedbackValues(ControllerId, static_cast<int32>(Hand), HapticValues, bFromHapticInterface);
+#if WITH_MLSDK
+	return CurrMotionSourceControllerState[FMagicLeapMotionSourceNames::Control0].bIsConnected ||
+		CurrMotionSourceControllerState[FMagicLeapMotionSourceNames::Control1].bIsConnected ||
+		CurrMotionSourceControllerState[FMagicLeapMotionSourceNames::MobileApp].bIsConnected;
+#endif //WITH_MLSDK
+	return false;
 }
 
-//Convert directly from Force Feedback Values to haptic values.  Small = Frequency, Large = Amplitude
-void FMagicLeapController::SetChannelValues(int32 ControllerId, const FForceFeedbackValues &values)
+void FMagicLeapController::ReadConfigParams()
 {
-	bool bFromHapticInterface = false;
-	FHapticFeedbackValues HapticValues;
-
-	//body
-	HapticValues.Frequency = values.LeftSmall;
-	HapticValues.Amplitude = values.LeftLarge;
-	InternalSetHapticFeedbackValues(ControllerId, static_cast<int32>(EControllerHand::Left), HapticValues, bFromHapticInterface);
-
-	//touchpad
-	HapticValues.Frequency = values.RightSmall;
-	HapticValues.Amplitude = values.RightLarge;
-	InternalSetHapticFeedbackValues(ControllerId, static_cast<int32>(EControllerHand::Right), HapticValues, bFromHapticInterface);
-}
+#if WITH_MLSDK
+	// Pull hand-mapping preferences from config file. If there are none, the default (legacy)
+	// mapping of device 0 to right and device 1 to left will persist.
+	const UEnum* ControllerHandEnum = FindObject<UEnum>(ANY_PACKAGE, TEXT("EControllerHand"));
+	TArray<FString> ControllerHands;
+	GConfig->GetArray(TEXT("/Script/LuminRuntimeSettings.LuminRuntimeSettings"),
+		TEXT("ControllerHands"), ControllerHands, GEngineIni);
+	for (auto ControllerHand : ControllerHands)
+	{
+		// Remove the parentheses
+		ControllerHand.RemoveAt(0);
+		ControllerHand.RemoveAt(ControllerHand.Len() - 1);
+
+		// Get the mapping
+		TArray<FString> SingleMapping;
+		ControllerHand.ParseIntoArray(SingleMapping, TEXT("="), true);
+		if (SingleMapping.Num() == 2)
+		{
+			auto HandIndex = ControllerHandEnum->GetValueByNameString(SingleMapping[0]);
+			if (HandIndex != INDEX_NONE)
+			{
+				ControllerMapper.MapHandToMotionSource(static_cast<EControllerHand>(HandIndex), FName(*SingleMapping[1]));
+			}
+			else
+			{
+				UE_LOG(LogMagicLeapController, Error,
+					TEXT("Invalid hand enum %s specified in ControllerHands array."),
+					*SingleMapping[0]);
+			}
+		}
+	}
 
-IHapticDevice* FMagicLeapController::GetHapticDevice()
-{
-	return this;
-}
+	// Pull trigger thresholds from config file
+	float FloatValueReceived = 0.0f;
+	GConfig->GetFloat(TEXT("/Script/Engine.InputSettings"),
+		TEXT("TriggerKeyIsConsideredPressed"), FloatValueReceived, GInputIni);
+	TriggerKeyIsConsideredPressed = FloatValueReceived;
 
-bool FMagicLeapController::IsGamepadAttached() const
-{
-#if WITH_MLSDK
-	if (bIsInputStateValid)
+	GConfig->GetFloat(TEXT("/Script/Engine.InputSettings"),
+		TEXT("TriggerKeyIsConsideredReleased"), FloatValueReceived, GInputIni);
+	TriggerKeyIsConsideredReleased = FloatValueReceived;
+
+	// Pull tracking preferences from config file
+	FString ConfigString;
+	GConfig->GetString(TEXT("/Script/LuminRuntimeSettings.LuminRuntimeSettings"),
+		TEXT("ControllerTrackingType"), ConfigString, GEngineIni);
+	if (ConfigString.Len() > 0)
 	{
-		// TODO: replace with flags set in OnControllerConnected and OnControllerDisconnected.
-		bool bDeviceConnected = false;
-		const int32* Controller_L = HandToControllerID.Find(EControllerHand::Left);
-		const int32* Controller_R = HandToControllerID.Find(EControllerHand::Right);
+		const UEnum* TrackingTypeEnum = FindObject<UEnum>(ANY_PACKAGE, TEXT("ETrackingStatus"));
 
-		if (Controller_L != nullptr)
+		auto TrackingTypeIndex = TrackingTypeEnum->GetValueByNameString(ConfigString);
+		if (TrackingTypeIndex != INDEX_NONE)
 		{
-			bDeviceConnected = bDeviceConnected || InputState[*Controller_L].is_connected;
+			switch (static_cast<ETrackingStatus>(TrackingTypeIndex))
+			{
+			case ETrackingStatus::NotTracked:
+				ControllerDof = MLInputControllerDof_None;
+				break;
+			case ETrackingStatus::InertialOnly:
+				ControllerDof = MLInputControllerDof_3;
+				break;
+			case ETrackingStatus::Tracked:
+			default:
+				ControllerDof = MLInputControllerDof_6;
+				break;
+			}
 		}
-		if (Controller_R != nullptr)
+		else
 		{
-			bDeviceConnected = bDeviceConnected || InputState[*Controller_R].is_connected;
+			UE_LOG(LogMagicLeapController, Error,
+				TEXT("Invalid ControllerTrackingType %s specified."),
+				*ConfigString);
 		}
-		return bDeviceConnected;
+	}
+
+	const static UEnum* TrackingModeEnum = FindObject<UEnum>(ANY_PACKAGE, TEXT("EMLControllerTrackingMode"));
+	GConfig->GetString(TEXT("/Script/LuminRuntimeSettings.LuminRuntimeSettings"),
+		TEXT("ControllerTrackingMode"), ConfigString, GEngineIni);
+	if (ConfigString.Len() > 0)
+	{
+		TrackingMode = static_cast<EMLControllerTrackingMode>(TrackingModeEnum->GetValueByNameString(ConfigString));
 	}
 #endif //WITH_MLSDK
-	return false;
 }
 
 void FMagicLeapController::Enable()
 {
 #if WITH_MLSDK
-	// Default to CFUID
-	TrackingMode = EMLControllerTrackingMode::CoordinateFrameUID;
+	MLResult Result = MLResult_Ok;
 
-	// Pull preference from config file
-	const static UEnum* TrackingModeEnum = FindObject<UEnum>(ANY_PACKAGE, TEXT("EMLControllerTrackingMode"));
+	// Must be done at Enable b/c we need packages to load to read enums
+	ReadConfigParams();
 
-	FString EnumVal;
-	GConfig->GetString(TEXT("/Script/LuminRuntimeSettings.LuminRuntimeSettings"),
-		TEXT("ControllerTrackingMode"), EnumVal, GEngineIni);
+	FMemory::Memset(&InputControllerState, 0, sizeof(InputControllerState));
 
-	if (EnumVal.Len() > 0)
+	// Attempt to create the Controller Tracker, regardless of chosen mode,
+	// so we can switch on the fly
+	MLControllerConfiguration ControllerConfig = { false };
+
+	switch (ControllerDof)
 	{
-		TrackingMode = static_cast<EMLControllerTrackingMode>(TrackingModeEnum->GetValueByNameString(EnumVal));
+	case MLInputControllerDof_3:
+		ControllerConfig.enable_imu3dof = true;
+		break;
+	case MLInputControllerDof_6:
+		ControllerConfig.enable_fused6dof = true;
+		break;
+	case MLInputControllerDof_None:
+		break;
 	}
 
-	// Attempt to create the Controller Tracker. We always do this b/c we do not want to create
-	// the tracker on the fly, and the mode can be changed on the fly.
-	if (!MLHandleIsValid(ControllerTracker))
+	Result = MLControllerCreate(ControllerConfig, &ControllerTracker);
+
+	if (Result != MLResult_Ok)
 	{
-		MLControllerConfiguration ControllerConfig = { false };
+		UE_LOG(LogMagicLeapController, Error,
+			TEXT("MLControllerCreate failed with error %s."),
+			UTF8_TO_TCHAR(MLGetResultString(Result)));
+		ControllerTracker = ML_INVALID_HANDLE;
 
-		ControllerConfig.enable_fused6dof = true;
+		// Revert to input
+		TrackingMode = EMLControllerTrackingMode::InputService;
+	}
 
-		MLResult Result = MLControllerCreate(ControllerConfig, &ControllerTracker);
+#if PLATFORM_LUMIN
+	// On-platform we pull the input tracker from LuminApplication
+	InputTracker = static_cast<FLuminApplication *>
+		(FSlateApplication::Get().GetPlatformApplication().Get())->GetInputTracker();
+#else
+	// For ML Remote (PIE) we need to create the input tracker here,
+	// as LuminApplication is not created.
+	MLInputConfiguration InputConfig = { };
 
-		if (Result != MLResult_Ok)
-		{
-			UE_LOG(LogMagicLeapController, Error, 
-				TEXT("MLControllerCreate failed with error %s."), 
-				UTF8_TO_TCHAR(MLGetResultString(Result)));
-			ControllerTracker = ML_INVALID_HANDLE;
+	InputConfig.dof[0] = InputConfig.dof[1] = ControllerDof;
 
-			// Revert to input
-			TrackingMode = EMLControllerTrackingMode::InputService;
-		}
+	Result = MLInputCreate(&InputConfig, &InputTracker);
+	if (Result != MLResult_Ok)
+	{
+		UE_LOG(LogMagicLeapController, Error, 
+			TEXT("MLInputCreate failed with error %s."), 
+			UTF8_TO_TCHAR(MLGetResultString(Result)));
 	}
+#endif //PLATFORM_LUMIN
 
-	// Attempt to create the Input Tracker
-	if (!MLHandleIsValid(InputTracker))
+	// Set controller button/touchpad callbacks on valid input tracker
+	if (MLHandleIsValid(InputTracker))
 	{
-		MLResult Result = MLInputCreate(nullptr, &InputTracker);
-		if (Result == MLResult_Ok)
-		{
-			FMemory::Memset(&InputState, 0, sizeof(InputState));
-			Result = MLInputSetControllerCallbacks(InputTracker, &InputControllerCallbacks, this);
-			if (Result == MLResult_Ok)
-			{
-				Result = MLInputSetKeyboardCallbacks(InputTracker, &InputKeyboardCallbacks, this);
-				if (Result != MLResult_Ok)
-				{
-					UE_LOG(LogMagicLeapController, Error, 
-						TEXT("MLInputSetKeyboardCallbacks failed with error %s."), 
-						UTF8_TO_TCHAR(MLGetResultString(Result)));
-				}
-
-				// Poll to get startup status
-				UpdateTrackerData();
-			}
-			else
-			{
-				UE_LOG(LogMagicLeapController, Error, 
-					TEXT("MLInputSetControllerCallbacks failed with error %s."), 
-					UTF8_TO_TCHAR(MLGetResultString(Result)));
-			}
-		}
-		else
+		Result = MLInputSetControllerCallbacks(InputTracker, &InputControllerCallbacks, this);
+		if (Result != MLResult_Ok)
 		{
-			UE_LOG(LogMagicLeapController, Error, 
-				TEXT("MLInputCreate failed with error %s."), 
+			UE_LOG(LogMagicLeapController, Error,
+				TEXT("MLInputSetControllerCallbacks failed with error %s."),
 				UTF8_TO_TCHAR(MLGetResultString(Result)));
 		}
 	}
+
+	// Poll to get startup status
+	UpdateTrackerData();
+
 #endif //WITH_MLSDK
 }
 
@@ -643,6 +689,7 @@ bool FMagicLeapController::SupportsExplicitEnable() const
 void FMagicLeapController::Disable()
 {
 #if WITH_MLSDK
+#if !PLATFORM_LUMIN
 	if (MLHandleIsValid(InputTracker))
 	{
 		MLResult Result = MLInputDestroy(InputTracker);
@@ -652,178 +699,81 @@ void FMagicLeapController::Disable()
 				TEXT("MLInputDestroy failed with error %s!"), 
 				UTF8_TO_TCHAR(MLGetResultString(Result)));
 		}
-		InputTracker = ML_INVALID_HANDLE;
-		bIsInputStateValid = false;
 	}
+#endif //PLATFORM_LUMIN
+	InputTracker = ML_INVALID_HANDLE;
+
 	if (MLHandleIsValid(ControllerTracker))
 	{
 		MLResult Result = MLControllerDestroy(ControllerTracker);
 		if (Result != MLResult_Ok)
 		{
-			UE_LOG(LogMagicLeapController, Error, 
-				TEXT("MLControllerDestroy failed with error %s!"), 
+			UE_LOG(LogMagicLeapController, Error,
+				TEXT("MLControllerDestroy failed with error %s!"),
 				UTF8_TO_TCHAR(MLGetResultString(Result)));
 		}
-		ControllerTracker = ML_INVALID_HANDLE;
 	}
-#endif //WITH_MLSDK
-}
+	ControllerTracker = ML_INVALID_HANDLE;
 
-void FMagicLeapController::SetHapticFeedbackValues(int32 ControllerId, int32 Hand, const FHapticFeedbackValues& Values)
-{
-	bool bFromHapticInterface = true;
-	InternalSetHapticFeedbackValues(ControllerId, Hand, Values, bFromHapticInterface);
+	bIsInputStateValid = false;
+#endif //WITH_MLSDK
 }
 
-void FMagicLeapController::InternalSetHapticFeedbackValues(int32 ControllerId, int32 Hand, const FHapticFeedbackValues& Values, bool bFromHapticInterface)
+bool FMagicLeapController::GetControllerOrientationAndPosition(const int32 ControllerIndex, const FName MotionSource, FRotator& OutOrientation, FVector& OutPosition, float WorldToMetersScale) const
 {
-	if (ControllerId != DeviceIndex)
-	{
-		return;
-	}
-
-	const int32* ControllerIndex = HandToControllerID.Find((EControllerHand)Hand);
-	if (ControllerIndex != nullptr)
+#if WITH_MLSDK
+	if (ControllerIndex == DeviceIndex)
 	{
-		//get the shadow state from the controller id
-		FMagicLeapControllerState* ControllerState = ControllerIDToControllerState.Find(*ControllerIndex);
-		if ((ControllerState != nullptr) && (ControllerState->bIsConnected) && (FApp::HasVRFocus()))
-		{
-			//early out if this is force feedback trying to run while haptics are running
-			if (ControllerState->bPlayingHapticEffect && !bFromHapticInterface)
-			{
-				return;
-			}
-
-			//clear to apply change to haptics
-			FHapticFeedbackBuffer* HapticBuffer = Values.HapticBuffer;
-			if (HapticBuffer) // && HapticBuffer->SamplingRate == SampleRateHz)
-			{
-				//@TODO - Setup buffer transfer
-				//if (bFromHapticInterface)
-				//{
-				//	ControllerState.bPlayingHapticEffect = (Amplitude != 0.f) && (Frequency != 0.f);
-				//}
-			}
-			else
-			{
-				float FreqMin, FreqMax = 0.f;
-				GetHapticFrequencyRange(FreqMin, FreqMax);
-
-				const float Frequency = FMath::Lerp(FreqMin, FreqMax, FMath::Clamp(Values.Frequency, 0.f, 1.f));
-				const float Amplitude = Values.Amplitude * GetHapticAmplitudeScale();
-
-				if (ControllerState->HapticAmplitude != Amplitude || ControllerState->HapticFrequency != Frequency)
-				{
-					ControllerState->HapticAmplitude = Amplitude;
-					ControllerState->HapticFrequency = Frequency;
-
-					UE_LOG(LogMagicLeapController, Warning, TEXT("Disabling haptic feedback pending a lower level API.  please use UMagicLeapControllerFunctionLibrary::PlayControllerHapticFeedback for now!"));
-
-					if ((EControllerHand(Hand) == EControllerHand::Left))
-					{
-						// TODO: What to do with duration?
-						//MLResult Result = MLInputHapticsStartControllerBody(InputTracker, static_cast<uint8>(*ControllerIndex), MLInputControllerFeedbackPatternVibe_Buzz, intensity, 8);
-						//if (Result != MLResult_Ok)
-						//{
-						//	UE_LOG(LogMagicLeapController, Error, TEXT("MLInputHapticsStartControllerBody() failed."));
-						//}
-					}
-					else
-					{
-						// TODO: What to do with duration?
-						//MLResult Result = MLInputHapticsStartControllerTouchpad(InputTracker, static_cast<uint8>(*Controller), MLInputControllerFeedbackPatternVibe_Buzz, intensity, 8);
-						//if (Result != MLResult_Ok)
-						//{
-						//	UE_LOG(LogMagicLeapController, Error, TEXT("MLInputHapticsStartControllerTouchpad() failed."));
-						//}
-					}
-
-					//only toggle when called from the haptics interface, not force feedback
-					if (bFromHapticInterface)
-					{
-						ControllerState->bPlayingHapticEffect = (Amplitude != 0.f) && (Frequency != 0.f);
-					}
-				}
-			}
-		}
-		else
+		const auto& ControllerState = CurrMotionSourceControllerState.Find(MotionSource);
+		if (ControllerState != nullptr)
 		{
-#if PLATFORM_LUMIN
-			static int32 bHasErrored = 0;
-			if (!bHasErrored)
-			{
-				bHasErrored = true;
-				UE_LOG(LogMagicLeapController, Error, TEXT("Haptic controller not attached"));
-			}
-#endif
+			OutPosition = ControllerState->Transform.GetLocation();
+			OutOrientation = ControllerState->Transform.GetRotation().Rotator();
+			return true;
 		}
 	}
-}
-
-void FMagicLeapController::GetHapticFrequencyRange(float& MinFrequency, float& MaxFrequency) const
-{
-	MinFrequency = 0.0f;
-	MaxFrequency = 1.0f;
-}
-
-float FMagicLeapController::GetHapticAmplitudeScale() const
-{
-	return 1.0f;
+#endif //WITH_MLSDK
+	return FXRMotionControllerBase::GetControllerOrientationAndPosition(ControllerIndex, MotionSource, OutOrientation, OutPosition, WorldToMetersScale);
 }
 
 bool FMagicLeapController::GetControllerOrientationAndPosition(const int32 ControllerIndex, const EControllerHand DeviceHand, FRotator& OutOrientation, FVector& OutPosition, float WorldToMetersScale) const
 {
-	bool bControllerTracked = false;
+#if WITH_MLSDK
+	// Do not call our function with a remap as that can enter an infinite loop due to remapping in FXRMotionControlerBase
 	if (ControllerIndex == DeviceIndex)
 	{
-		if (GetControllerTrackingStatus(ControllerIndex, DeviceHand) != ETrackingStatus::NotTracked)
+		auto MotionSource = ControllerMapper.GetMotionSourceForHand(DeviceHand);
+		if (MotionSource != FMagicLeapMotionSourceNames::Unknown)
 		{
-			const FTransform* ControllerTransform = &LeftControllerTransform;
-
-			if (DeviceHand == EControllerHand::Right)
-			{
-				ControllerTransform = &RightControllerTransform;
-			}
-
-			OutPosition = ControllerTransform->GetLocation();
-			OutOrientation = ControllerTransform->GetRotation().Rotator();
-
-			bControllerTracked = true;
+			return GetControllerOrientationAndPosition(ControllerIndex,
+				MotionSource, OutOrientation, OutPosition, WorldToMetersScale);
 		}
 	}
-
-	return bControllerTracked;
+#endif //WITH_MLSDK
+	return false;
 }
 
-ETrackingStatus FMagicLeapController::GetControllerTrackingStatus(const int32 ControllerIndex, const EControllerHand DeviceHand) const
+ETrackingStatus FMagicLeapController::GetControllerTrackingStatus(const int32 ControllerIndex, const FName MotionSource) const
 {
-	ETrackingStatus status = ETrackingStatus::NotTracked;
 #if WITH_MLSDK
 	if (ControllerIndex == DeviceIndex)
 	{
-		const int32* ControllerID = HandToControllerID.Find(DeviceHand);
-		if (ControllerID != nullptr)
+		const auto& ControllerState = CurrMotionSourceControllerState.Find(MotionSource);
+		if (ControllerState != nullptr)
 		{
-			if (bIsInputStateValid && InputState[*ControllerID].is_connected)
-			{
-				switch (InputState[*ControllerID].dof)
-				{
-				case MLInputControllerDof_3:
-					status = ETrackingStatus::InertialOnly;
-					break;
-				case MLInputControllerDof_6:
-					status = ETrackingStatus::Tracked;
-					break;
-				default:
-					status = ETrackingStatus::NotTracked;
-					break;
-				}
-			}
+			return ControllerState->TrackingStatus;
 		}
 	}
 #endif //WITH_MLSDK
-	return status;
+	return FXRMotionControllerBase::GetControllerTrackingStatus(ControllerIndex, MotionSource);
+}
+
+ETrackingStatus FMagicLeapController::GetControllerTrackingStatus(const int32 ControllerIndex, const EControllerHand DeviceHand) const
+{
+#if WITH_MLSDK
+	return GetControllerTrackingStatus(ControllerIndex, ControllerMapper.GetMotionSourceForHand(DeviceHand));
+#endif //WITH_MLSDK
+	return ETrackingStatus::NotTracked;
 }
 
 FName FMagicLeapController::GetMotionControllerDeviceTypeName() const
@@ -832,227 +782,254 @@ FName FMagicLeapController::GetMotionControllerDeviceTypeName() const
 	return DefaultName;
 }
 
-void FMagicLeapController::UpdateControllerTransformFromInputTracker(const FAppFramework& AppFramework, FTransform& ControllerTransform, EControllerHand ControllerHand)
+void FMagicLeapController::EnumerateSources(TArray<FMotionControllerSource>& SourcesOut) const
 {
-#if WITH_MLSDK
-	int32 ControllerID = *HandToControllerID.Find(ControllerHand);
-	ETrackingStatus ControllerTrackingStatus = GetControllerTrackingStatus(DeviceIndex, ControllerHand);
-	if (ControllerTrackingStatus == ETrackingStatus::Tracked)
-	{
-		ControllerTransform.SetLocation(MagicLeap::ToFVector(InputState[ControllerID].position, AppFramework.GetWorldToMetersScale()));
-		ControllerTransform.SetRotation(MagicLeap::ToFQuat(InputState[ControllerID].orientation));
-	}
-	else if (ControllerTrackingStatus == ETrackingStatus::InertialOnly)
-	{
-		ControllerTransform.SetRotation(MagicLeap::ToFQuat(InputState[ControllerID].orientation));
-	}
-
-	if (ControllerTransform.ContainsNaN())
-	{
-		UE_LOG(LogMagicLeapController, Error, TEXT("Transform for controller index %d has NaNs."), ControllerID);
-		InputState[ControllerID].dof = MLInputControllerDof_None;
-	}
-	else
-	{
-		if (!ControllerTransform.GetRotation().IsNormalized())
-		{
-			FQuat rotation = ControllerTransform.GetRotation();
-			rotation.Normalize();
-			ControllerTransform.SetRotation(rotation);
-		}
-	}
-#endif //WITH_MLSDK
+	SourcesOut.Add(FMagicLeapMotionSourceNames::Control0);
+	SourcesOut.Add(FMagicLeapMotionSourceNames::Control1);
+	SourcesOut.Add(FMagicLeapMotionSourceNames::MobileApp);
 }
 
-#if WITH_MLSDK
-void FMagicLeapController::UpdateControllerTransformFromControllerTracker(const FAppFramework& AppFramework, const MLControllerSystemState& ControllerSystemState, FTransform& ControllerTransform, int32 InDeviceIndex)
+void FMagicLeapController::UpdateControllerStateFromInputTracker(const FAppFramework& AppFramework, FName MotionSource)
 {
-	const auto& ControllerStream = ControllerSystemState.
-		controller_state[InDeviceIndex].stream[MLControllerMode_Fused6Dof];
+#if WITH_MLSDK
+	checkf(CurrMotionSourceControllerState.Contains(MotionSource),
+		TEXT("UpdateControllerStateFromInputTracker was asked for non-ML motion source!"));
+	checkf(PrevMotionSourceControllerState.Contains(MotionSource), TEXT("Unpossible"));
 
-	if (ControllerStream.is_active)
-	{
-		EFailReason FailReason = EFailReason::None;
-		if (!AppFramework.GetTransform(ControllerStream.coord_frame_controller, ControllerTransform, FailReason))
-		{
-			UE_LOG(LogMagicLeapController, Error, 
-				TEXT("UpdateControllerTransformFromControllerTracker: AppFramework."
-				"GetTransform returned false, fail reason = %d."), 
-				static_cast<uint32>(FailReason));
-		}
-	}
-}
-#endif //WITH_MLSDK
+	auto& CurrControllerState = CurrMotionSourceControllerState[MotionSource];
+	auto& PrevControllerState = PrevMotionSourceControllerState[MotionSource];
 
-void FMagicLeapController::UpdateTrackerData()
-{
-#if WITH_MLSDK
-	if (!IMagicLeapPlugin::Get().IsMagicLeapHMDValid())
-	{
-		return;
-	}
+	// Advance frame
+	PrevControllerState = CurrControllerState;
 
-	const FAppFramework& AppFramework = static_cast<FMagicLeapHMD*>
-		(GEngine->XRSystem->GetHMDDevice())->GetAppFrameworkConst();
-	if (!AppFramework.IsInitialized())
+	// Get platform input state for this motion source.
+	int InputStateIndex = ControllerMapper.GetInputControllerIndexForMotionSource(MotionSource);
+	if (InputStateIndex != -1)
 	{
-		return;
-	}
+		const auto& InputState = InputControllerState[InputStateIndex];
 
-	// First pull data from input tracker. Note that this is not conditional based on the tracking
-	// type because we also need to get button, touchpad, etc.
-	if (MLHandleIsValid(InputTracker))
-	{
-		FMemory::Memcpy(&OldInputState, &InputState, sizeof(InputState));
+		// TODO: connect/disconnect events?
+		CurrControllerState.bIsConnected = InputState.is_connected;
 
-		bIsInputStateValid = MLInputGetControllerState(InputTracker, InputState) == MLResult_Ok;
+		// Touch activity, coordinates, and force
+		for (auto TouchIndex = 0; TouchIndex < MLInput_MaxControllerTouchpadTouches; ++ TouchIndex)
+		{
+			CurrControllerState.bTouchActive[TouchIndex] = InputState.is_touch_active[TouchIndex];
+
+			CurrControllerState.TouchPosAndForce[TouchIndex].
+				Set(InputState.touch_pos_and_force[TouchIndex].x,
+					InputState.touch_pos_and_force[TouchIndex].y,
+					InputState.touch_pos_and_force[TouchIndex].z);
+		}
+
+		// Analog trigger
+		CurrControllerState.TriggerAnalog = InputState.trigger_normalized;
 
-		if (bIsInputStateValid)
+		// Dof
+		switch (InputState.dof)
 		{
-			UpdateControllerTransformFromInputTracker(AppFramework, 
-				LeftControllerTransform, EControllerHand::Left);
-			UpdateControllerTransformFromInputTracker(AppFramework, 
-				RightControllerTransform, EControllerHand::Right);
+			case MLInputControllerDof_3:
+			{
+				CurrControllerState.TrackingStatus = ETrackingStatus::InertialOnly;
+				CurrControllerState.Transform.SetLocation(FVector::ZeroVector);
+				CurrControllerState.Transform.SetRotation(MagicLeap::ToFQuat(InputState.orientation));
+				break;
+			}
+			case MLInputControllerDof_6:
+			{
+				CurrControllerState.TrackingStatus = ETrackingStatus::Tracked;
+				CurrControllerState.Transform.SetLocation(MagicLeap::ToFVector(InputState.position,
+					AppFramework.GetWorldToMetersScale()));
+				CurrControllerState.Transform.SetRotation(MagicLeap::ToFQuat(InputState.orientation));
+				break;
+			}
+			default:
+			{
+				CurrControllerState.TrackingStatus = ETrackingStatus::NotTracked;
+				CurrControllerState.Transform.SetIdentity();
+				break;
+			}
 		}
-	}
 
-	// If mode is set to CFUID tracking overwrite the input data
-	if (bIsInputStateValid && MLHandleIsValid(ControllerTracker) && (TrackingMode == EMLControllerTrackingMode::CoordinateFrameUID))
-	{
-		MLControllerSystemState ControllerSystemState;
-		MLResult Result = MLControllerGetState(ControllerTracker, &ControllerSystemState);
-		if (MLResult_Ok == Result)
+		// Fixup transform
+		if (CurrControllerState.Transform.ContainsNaN())
+		{
+			UE_LOG(LogMagicLeapController, Error, TEXT("Transform for input state index %d has NaNs."), InputStateIndex);
+			CurrControllerState.TrackingStatus = ETrackingStatus::NotTracked;
+			CurrControllerState.Transform.SetIdentity();
+		}
+		else if (!CurrControllerState.Transform.GetRotation().IsNormalized())
 		{
-			int32 ControllerID;
+			FQuat rotation = CurrControllerState.Transform.GetRotation();
+			rotation.Normalize();
+			CurrControllerState.Transform.SetRotation(rotation);
+		}
 
-			ControllerID = HandToControllerID[EControllerHand::Left];
-			if (InputState[ControllerID].type == MLInputControllerType_Device)
+		// Generated button events
+		auto Hand = ControllerMapper.GetHandForMotionSource(MotionSource);
+		if (Hand != EControllerHand::ControllerHand_Count)
+		{
+			// Touch0 activate/deactivate
+			if (CurrControllerState.bTouchActive[0] && !PrevControllerState.bTouchActive[0])
 			{
-				UpdateControllerTransformFromControllerTracker(AppFramework, ControllerSystemState,
-					LeftControllerTransform, InputState[ControllerID].hardware_index);
+				PendingButtonEvents.Enqueue(MakeTuple(MLTouchToUnrealThumbstickButton(Hand), true));
+			}
+			else if (PrevControllerState.bTouchActive[0] && !CurrControllerState.bTouchActive[0])
+			{
+				PendingButtonEvents.Enqueue(MakeTuple(MLTouchToUnrealThumbstickButton(Hand), false));
 			}
 
-			ControllerID = HandToControllerID[EControllerHand::Right];
-			if (InputState[ControllerID].type == MLInputControllerType_Device)
+			// Convert trigger value to trigger press/release events
+			const bool IsTriggerKeyPressing = (CurrControllerState.TriggerAnalog >
+				TriggerKeyIsConsideredPressed) && !PrevControllerState.bTriggerKeyPressing;
+			const bool IsTriggerKeyReleasing = (CurrControllerState.TriggerAnalog <
+				TriggerKeyIsConsideredReleased) && PrevControllerState.bTriggerKeyPressing;
+
+			if (IsTriggerKeyPressing)
 			{
-				UpdateControllerTransformFromControllerTracker(AppFramework, ControllerSystemState,
-					RightControllerTransform, InputState[ControllerID].hardware_index);
+				PendingButtonEvents.Enqueue(MakeTuple(MLTriggerToUnrealTriggerKey(Hand), true));
+				CurrControllerState.bTriggerKeyPressing = true;
+			}
+			else if (IsTriggerKeyReleasing)
+			{
+				PendingButtonEvents.Enqueue(MakeTuple(MLTriggerToUnrealTriggerKey(Hand), false));
+				CurrControllerState.bTriggerKeyPressing = false;
 			}
 		}
-		else
-		{
-			UE_LOG(LogMagicLeapController, Error, 
-				TEXT("MLControllerGetState failed with error %s."), 
-				UTF8_TO_TCHAR(MLGetResultString(Result)));
-		}
+
+	}
+	else
+	{
+		CurrControllerState = { };
 	}
 #endif //WITH_MLSDK
 }
 
-bool FMagicLeapController::IsInputStateValid() const
+void FMagicLeapController::UpdateControllerStateFromControllerTracker(const FAppFramework& AppFramework, FName MotionSource)
 {
-	return bIsInputStateValid;
-}
+#if WITH_MLSDK
+	// Index of the stream we're reading
+	int32 StreamIndex = -1;
 
-bool FMagicLeapController::GetControllerMapping(int32 ControllerIndex, EControllerHand& Hand) const
-{
-	const EControllerHand* possibleHand = ControllerIDToHand.Find(ControllerIndex);
-	if (possibleHand != nullptr)
+	switch (ControllerDof)
 	{
-		Hand = *possibleHand;
-		return true;
+	case MLInputControllerDof_3:
+		StreamIndex = MLControllerMode_Imu3Dof;
+		break;
+	case MLInputControllerDof_6:
+		StreamIndex = MLControllerMode_Fused6Dof;
+		break;
+	default:
+		break;
 	}
-	Hand = EControllerHand::Special_9;
-	return false;
-}
 
-void FMagicLeapController::InvertControllerMapping()
-{
-	EControllerHand tempHand = ControllerIDToHand[0];
-	ControllerIDToHand[0] = ControllerIDToHand[1];
-	ControllerIDToHand[1] = tempHand;
+	if (StreamIndex != -1)
+	{
+		checkf(MotionSource == FMagicLeapMotionSourceNames::Control0 || MotionSource == FMagicLeapMotionSourceNames::Control1,
+			TEXT("UpdateControllerStateFromControllerTracker was asked for non-control motion source!"));
+		auto& ControllerState = CurrMotionSourceControllerState[MotionSource];
 
-	int32 tempID = HandToControllerID[EControllerHand::Left];
-	HandToControllerID[EControllerHand::Left] = HandToControllerID[EControllerHand::Right];
-	HandToControllerID[EControllerHand::Right] = tempID;
-}
+		// Hardware index of control
+		int32 ControlIndex = 0;
+		if (MotionSource == FMagicLeapMotionSourceNames::Control1)
+		{
+			ControlIndex = 1;
+		}
 
-EMLControllerType FMagicLeapController::GetMLControllerType(EControllerHand Hand) const
-{
-#if WITH_MLSDK
-	const int32* ControllerID = HandToControllerID.Find(Hand);
-	if (ControllerID != nullptr)
-	{
-		switch (InputState[*ControllerID].type)
+		const auto& ControllerStream = ControllerSystemState.controller_state[ControlIndex].stream[StreamIndex];
+
+		if (ControllerStream.is_active)
+		{
+			if (StreamIndex == MLControllerMode_Imu3Dof)
+			{
+				ControllerState.TrackingStatus = ETrackingStatus::InertialOnly;
+			}
+			else
+			{
+				ControllerState.TrackingStatus = ETrackingStatus::Tracked;
+			}
+
+			EFailReason FailReason = EFailReason::None;
+			if (!AppFramework.GetTransform(ControllerStream.
+				coord_frame_controller, ControllerState.Transform, FailReason))
+			{
+				UE_LOG(LogMagicLeapController, Error,
+					TEXT("UpdateControllerStateFromControllerTracker: AppFramework."
+						"GetTransform returned false, fail reason = %d."),
+					static_cast<uint32>(FailReason));
+			}
+		}
+		else
 		{
-		case MLInputControllerType_Device:
-			return EMLControllerType::Device;
-		case MLInputControllerType_MobileApp:
-			return EMLControllerType::MobileApp;
+			ControllerState.TrackingStatus = ETrackingStatus::NotTracked;
 		}
 	}
 #endif //WITH_MLSDK
-	return EMLControllerType::None;
 }
 
-bool FMagicLeapController::PlayControllerLED(EControllerHand Hand, EMLControllerLEDPattern LEDPattern, EMLControllerLEDColor LEDColor, float DurationInSec)
+void FMagicLeapController::UpdateTrackerData()
 {
 #if WITH_MLSDK
-	if (IsGamepadAttached())
+	if (!IMagicLeapPlugin::Get().IsMagicLeapHMDValid())
 	{
-		const int32* ControllerID = HandToControllerID.Find(Hand);
-		if (ControllerID != nullptr)
-		{
-			return MLInputStartControllerFeedbackPatternLED(InputTracker, static_cast<uint8>(*ControllerID), UnrealToMLPatternLED(LEDPattern), UnrealToMLColorLED(LEDColor), static_cast<uint32>(DurationInSec * 1000)) == MLResult_Ok;
-		}
+		return;
 	}
-	else
+
+	const FAppFramework& AppFramework = static_cast<FMagicLeapHMD*>
+		(GEngine->XRSystem->GetHMDDevice())->GetAppFrameworkConst();
+	if (!AppFramework.IsInitialized())
 	{
-		UE_LOG(LogMagicLeapController, Error, TEXT("LED controller not attached"));
+		return;
 	}
-#endif //WITH_MLSDK
 
-	return false;
-}
-
-bool FMagicLeapController::PlayControllerLEDEffect(EControllerHand Hand, EMLControllerLEDEffect LEDEffect, EMLControllerLEDSpeed LEDSpeed, EMLControllerLEDPattern LEDPattern, EMLControllerLEDColor LEDColor, float DurationInSec)
-{
-#if WITH_MLSDK
-	if (IsGamepadAttached())
+	// First pull data from input tracker. Note that this is not conditional based on the tracking
+	// type because we also need to get buttons, touchpad, etc.
+	if (MLHandleIsValid(InputTracker))
 	{
-		const int32* ControllerID = HandToControllerID.Find(Hand);
-		if (ControllerID != nullptr)
+		MLResult Result = MLInputGetControllerState(InputTracker, InputControllerState);
+
+		if (MLResult_Ok == Result)
 		{
-			return MLInputStartControllerFeedbackPatternEffectLED(InputTracker, static_cast<uint8>(*ControllerID), UnrealToMLEffectLED(LEDEffect), UnrealToMLSpeedLED(LEDSpeed), UnrealToMLPatternLED(LEDPattern), UnrealToMLColorLED(LEDColor), static_cast<uint32>(DurationInSec * 1000)) == MLResult_Ok;
+			bIsInputStateValid = true;
+
+			ControllerMapper.UpdateMotionSourceInputIndexPairing(InputControllerState);
+			UpdateControllerStateFromInputTracker(AppFramework, FMagicLeapMotionSourceNames::Control0);
+			UpdateControllerStateFromInputTracker(AppFramework, FMagicLeapMotionSourceNames::Control1);
+			UpdateControllerStateFromInputTracker(AppFramework, FMagicLeapMotionSourceNames::MobileApp);
 		}
-	}
-	else
-	{
-		UE_LOG(LogMagicLeapController, Error, TEXT("LED controller not attached"));
-	}
-#endif //WITH_MLSDK
+		else
+		{
+			bIsInputStateValid = false;
 
-	return false;
-}
+			UE_LOG(LogMagicLeapController, Error,
+				TEXT("MLInputGetControllerState failed with error %s."),
+				UTF8_TO_TCHAR(MLGetResultString(Result)));
+		}
+	}
 
-bool FMagicLeapController::PlayControllerHapticFeedback(EControllerHand Hand, EMLControllerHapticPattern HapticPattern, EMLControllerHapticIntensity Intensity)
-{
-#if WITH_MLSDK
-	if (IsGamepadAttached())
+	// If mode is set to CFUID tracking overwrite the input system Dof data
+	if (TrackingMode == EMLControllerTrackingMode::CoordinateFrameUID)
 	{
-		const int32* ControllerID = HandToControllerID.Find(Hand);
-		if (ControllerID != nullptr)
+		// We need to have valid input state in order to do this, because we need to be sure
+		// what we are polling is a physical control
+		if (bIsInputStateValid && MLHandleIsValid(ControllerTracker))
 		{
-			return MLInputStartControllerFeedbackPatternVibe(InputTracker, static_cast<uint8>(*ControllerID), UnrealToMLPatternVibe(HapticPattern), UnrealToMLHapticIntensity(Intensity)) == MLResult_Ok;
+			MLResult Result = MLControllerGetState(ControllerTracker, &ControllerSystemState);
+
+			if (MLResult_Ok == Result)
+			{
+				UpdateControllerStateFromControllerTracker(AppFramework, FMagicLeapMotionSourceNames::Control0);
+				UpdateControllerStateFromControllerTracker(AppFramework, FMagicLeapMotionSourceNames::Control1);
+			}
+			else
+			{
+				UE_LOG(LogMagicLeapController, Error,
+					TEXT("MLControllerGetState failed with error %s."),
+					UTF8_TO_TCHAR(MLGetResultString(Result)));
+			}
 		}
 	}
-	else
-	{
-		UE_LOG(LogMagicLeapController, Error, TEXT("Haptic controller not attached"));
-	}
 #endif //WITH_MLSDK
-
-	return false;
 }
 
 bool FMagicLeapController::SetControllerTrackingMode(EMLControllerTrackingMode InTrackingMode)
@@ -1078,10 +1055,6 @@ EMLControllerTrackingMode FMagicLeapController::GetControllerTrackingMode()
 	{
 		return TrackingMode;
 	}
-	else
-	{
-		UE_LOG(LogMagicLeapController, Error, TEXT("Haptic controller not attached"));
-	}
 #endif //WITH_MLSDK
 
 	return EMLControllerTrackingMode::InputService;
@@ -1113,444 +1086,215 @@ void FMagicLeapController::AddKeys()
 	EKeys::AddKey(FKeyDetails(FMagicLeapKeys::Right_HomeButton, LOCTEXT("MagicLeap_Right_HomeButton", "ML (R) Home Button"), FKeyDetails::GamepadKey));
 }
 
-void FMagicLeapController::DebouncedButtonMessageHandler(bool NewButtonState, bool OldButtonState, const FName& ButtonName)
+bool FMagicLeapController::PlayLEDPattern(FName MotionSource, EMLControllerLEDPattern LEDPattern, EMLControllerLEDColor LEDColor, float DurationInSec)
 {
-	if (NewButtonState && !OldButtonState)
-	{
-		FMagicLeapHMD::EnableInput EnableInputFromHMD;
-		// fixes unreferenced parameter error for Windows package builds.
-		(void)EnableInputFromHMD;
-		MessageHandler->OnControllerButtonPressed(ButtonName, DeviceIndex, false);
-	}
-	else if (!NewButtonState && OldButtonState)
-	{
-		FMagicLeapHMD::EnableInput EnableInputFromHMD;
-		// fixes unreferenced parameter error for Windows package builds.
-		(void)EnableInputFromHMD;
-		MessageHandler->OnControllerButtonReleased(ButtonName, DeviceIndex, false);
-	}
-}
-
 #if WITH_MLSDK
-const FName& FMagicLeapController::MagicLeapButtonToUnrealButton(int32 ControllerID, MLInputControllerButton ml_button)
-{
-	static const FName empty;
-
-	const EControllerHand* possibleHand = ControllerIDToHand.Find(ControllerID);
-	if (possibleHand == nullptr)
+	if (IsGamepadAttached())
 	{
-		return empty;
-	}
+		auto InputControllerIndex = ControllerMapper.GetInputControllerIndexForMotionSource(MotionSource);
+		if (InputControllerIndex != 0xFF)
+		{
+			MLResult Result = MLInputStartControllerFeedbackPatternLED(InputTracker,
+				InputControllerIndex,
+				UnrealToMLPatternLED(LEDPattern),
+				UnrealToMLColorLED(LEDColor),
+				static_cast<uint32>(DurationInSec * 1000));
 
-	const EControllerHand hand = *possibleHand;
+			if (MLResult_Ok == Result)
+			{
+				return true;
+			}
 
-	switch (ml_button)
-	{
-	case MLInputControllerButton_Move:
-		if (hand == EControllerHand::Left)
-		{
-			return FMagicLeapControllerKeyNames::Left_MoveButton_Name;
-		}
-		return FMagicLeapControllerKeyNames::Right_MoveButton_Name;
-	case MLInputControllerButton_App:
-		if (hand == EControllerHand::Left)
-		{
-			return FMagicLeapControllerKeyNames::Left_AppButton_Name;
-		}
-		return FMagicLeapControllerKeyNames::Right_AppButton_Name;
-	case MLInputControllerButton_Bumper:
-		if (hand == EControllerHand::Left)
-		{
-			return FGamepadKeyNames::MotionController_Left_Shoulder;
+			UE_LOG(LogMagicLeapController, Error,
+				TEXT("MLInputStartControllerFeedbackPatternLED failed with error %s"),
+				UTF8_TO_TCHAR(MLGetResultString(Result)));
 		}
-		return FGamepadKeyNames::MotionController_Right_Shoulder;
-	case MLInputControllerButton_HomeTap:
-		if (hand == EControllerHand::Left)
+		else
 		{
-			return FMagicLeapControllerKeyNames::Left_HomeButton_Name;
+			UE_LOG(LogMagicLeapController, Error, TEXT("PlayLEDPattern requested on non-ML controller"));
 		}
-		return FMagicLeapControllerKeyNames::Right_HomeButton_Name;
-	default:
-		break;
 	}
-	return empty;
-}
+	else
+	{
+		UE_LOG(LogMagicLeapController, Error, TEXT("PlayLEDPattern: controller not attached"));
+	}
 #endif //WITH_MLSDK
 
+	return false;
+}
 
-const FName& FMagicLeapController::MagicLeapTouchToUnrealThumbstickAxis(EControllerHand Hand, uint32 TouchIndex)
+bool FMagicLeapController::PlayLEDEffect(FName MotionSource, EMLControllerLEDEffect LEDEffect, EMLControllerLEDSpeed LEDSpeed, EMLControllerLEDPattern LEDPattern, EMLControllerLEDColor LEDColor, float DurationInSec)
 {
-	static const FName empty;
-
-	switch (TouchIndex)
+#if WITH_MLSDK
+	if (IsGamepadAttached())
 	{
-	case 0:
-		if (Hand == EControllerHand::Left)
-		{
-			return FGamepadKeyNames::MotionController_Left_Thumbstick_X;
-		}
-		return FGamepadKeyNames::MotionController_Right_Thumbstick_X;
-	case 1:
-		if (Hand == EControllerHand::Left)
+		auto InputControllerIndex = ControllerMapper.GetInputControllerIndexForMotionSource(MotionSource);
+		if (InputControllerIndex != 0xFF)
 		{
-			return FGamepadKeyNames::MotionController_Left_Thumbstick_Y;
+			MLResult Result = MLInputStartControllerFeedbackPatternEffectLED(InputTracker,
+				InputControllerIndex,
+				UnrealToMLEffectLED(LEDEffect),
+				UnrealToMLSpeedLED(LEDSpeed),
+				UnrealToMLPatternLED(LEDPattern),
+				UnrealToMLColorLED(LEDColor),
+				static_cast<uint32>(DurationInSec * 1000));
+
+			if (MLResult_Ok == Result)
+			{
+				return true;
+			}
+
+			UE_LOG(LogMagicLeapController, Error,
+				TEXT("MLInputStartControllerFeedbackPatternEffectLED failed with error %s"),
+				UTF8_TO_TCHAR(MLGetResultString(Result)));
 		}
-		return FGamepadKeyNames::MotionController_Right_Thumbstick_Y;
-	case 2:
-		if (Hand == EControllerHand::Left)
+		else
 		{
-			return FMagicLeapControllerKeyNames::MotionController_Left_Thumbstick_Z_Name;
+			UE_LOG(LogMagicLeapController, Error, TEXT("PlayLEDEffect requested on non-ML controller"));
 		}
-		return FMagicLeapControllerKeyNames::MotionController_Right_Thumbstick_Z_Name;
-	default:
-		return empty;
 	}
+	else
+	{
+		UE_LOG(LogMagicLeapController, Error, TEXT("PlayLEDEffect: controller not attached"));
+	}
+#endif //WITH_MLSDK
 
-	return empty;
+	return false;
 }
 
-const FName& FMagicLeapController::MagicLeapTouchToUnrealThumbstickButton(EControllerHand Hand)
+bool FMagicLeapController::PlayHapticPattern(FName MotionSource, EMLControllerHapticPattern HapticPattern, EMLControllerHapticIntensity Intensity)
 {
-	static const FName empty;
+#if WITH_MLSDK
+	if (IsGamepadAttached())
+	{
+		auto InputControllerIndex = ControllerMapper.GetInputControllerIndexForMotionSource(MotionSource);
+		if (InputControllerIndex != 0xFF)
+		{
+			MLResult Result = MLInputStartControllerFeedbackPatternVibe(InputTracker,
+				InputControllerIndex,
+				UnrealToMLPatternVibe(HapticPattern),
+				UnrealToMLHapticIntensity(Intensity));
 
-	switch (Hand)
+			if (MLResult_Ok == Result)
+			{
+				return true;
+			}
+
+			UE_LOG(LogMagicLeapController, Error,
+				TEXT("MLInputStartControllerFeedbackPatternVibe failed with error %s"),
+				UTF8_TO_TCHAR(MLGetResultString(Result)));
+		}
+		else
+		{
+			UE_LOG(LogMagicLeapController, Error, TEXT("PlayHapticPattern requested on non-ML controller"));
+		}
+	}
+	else
 	{
-	case EControllerHand::Left:
-		return FGamepadKeyNames::MotionController_Left_Thumbstick;
-	case EControllerHand::Right:
-		return FGamepadKeyNames::MotionController_Right_Thumbstick;
+		UE_LOG(LogMagicLeapController, Error, TEXT("PlayHapticPattern: controller not attached"));
 	}
+#endif //WITH_MLSDK
 
-	return empty;
+	return false;
 }
 
-const FName& FMagicLeapController::MagicLeapTriggerToUnrealTriggerAxis(EControllerHand Hand)
+void FMagicLeapController::SendControllerEventsForHand(EControllerHand Hand)
 {
-	static const FName empty;
+#if WITH_MLSDK
+	const auto& MotionSource = ControllerMapper.GetMotionSourceForHand(Hand);
 
-	switch (Hand)
+	const auto& CurrControllerState = CurrMotionSourceControllerState.Find(MotionSource);
+	const auto& PrevControllerState = PrevMotionSourceControllerState.Find(MotionSource);
+
+	if (CurrControllerState != nullptr && CurrControllerState->bIsConnected)
 	{
-	case EControllerHand::Left:
-		return FGamepadKeyNames::MotionController_Left_TriggerAxis;
-	case EControllerHand::Right:
-		return FGamepadKeyNames::MotionController_Right_TriggerAxis;
-	}
+		checkf(PrevControllerState != nullptr, TEXT("Unpossible"));
 
-	return empty;
-}
+		// Analog touch coords
+		// Touch 0 maps to Motion Controller Thumbstick for hand
+		// Touch 1 is currently not available (we have nothing to map it to)
+		if (CurrControllerState->bTouchActive[0])
+		{
+			FMagicLeapHMD::EnableInput EnableInputFromHMD;
+			// fixes unreferenced parameter error for Windows package builds.
+			(void)EnableInputFromHMD;
+
+			MessageHandler->OnControllerAnalog(MLTouchToUnrealThumbstickAxis(Hand, 0),
+				DeviceIndex, CurrControllerState->TouchPosAndForce[0].X);
+			MessageHandler->OnControllerAnalog(MLTouchToUnrealThumbstickAxis(Hand, 1),
+				DeviceIndex, CurrControllerState->TouchPosAndForce[0].Y);
+			MessageHandler->OnControllerAnalog(MLTouchToUnrealThumbstickAxis(Hand, 2),
+				DeviceIndex, CurrControllerState->TouchPosAndForce[0].Z);
+		}
+		else
+		{
+			FMagicLeapHMD::EnableInput EnableInputFromHMD;
+			// fixes unreferenced parameter error for Windows package builds.
+			(void)EnableInputFromHMD;
 
-const FName& FMagicLeapController::MagicLeapTriggerToUnrealTriggerKey(EControllerHand Hand)
-{
-	static const FName empty;
-	switch (Hand)
-	{
-	case EControllerHand::Left:
-		return FGamepadKeyNames::MotionController_Left_Trigger;
-	case EControllerHand::Right:
-		return FGamepadKeyNames::MotionController_Right_Trigger;
-	}
-	return empty;
-}
+			MessageHandler->OnControllerAnalog(MLTouchToUnrealThumbstickAxis(Hand, 0), DeviceIndex, 0.0f);
+			MessageHandler->OnControllerAnalog(MLTouchToUnrealThumbstickAxis(Hand, 1), DeviceIndex, 0.0f);
+			MessageHandler->OnControllerAnalog(MLTouchToUnrealThumbstickAxis(Hand, 2), DeviceIndex, 0.0f);
+		}
 
-#if WITH_MLSDK
+		// Analog trigger
+		if (CurrControllerState->TriggerAnalog != PrevControllerState->TriggerAnalog)
+		{
+			FMagicLeapHMD::EnableInput EnableInputFromHMD;
+			// fixes unreferenced parameter error for Windows package builds.
+			(void)EnableInputFromHMD;
 
-MLInputControllerFeedbackPatternLED FMagicLeapController::UnrealToMLPatternLED(EMLControllerLEDPattern LEDPattern) const
-{
-	switch (LEDPattern)
-	{
-		case EMLControllerLEDPattern::None:
-			return MLInputControllerFeedbackPatternLED_None;
-		case EMLControllerLEDPattern::Clock01:
-			return MLInputControllerFeedbackPatternLED_Clock1;
-		case EMLControllerLEDPattern::Clock02:
-			return MLInputControllerFeedbackPatternLED_Clock2;
-		case EMLControllerLEDPattern::Clock03:
-			return MLInputControllerFeedbackPatternLED_Clock3;
-		case EMLControllerLEDPattern::Clock04:
-			return MLInputControllerFeedbackPatternLED_Clock4;
-		case EMLControllerLEDPattern::Clock05:
-			return MLInputControllerFeedbackPatternLED_Clock5;
-		case EMLControllerLEDPattern::Clock06:
-			return MLInputControllerFeedbackPatternLED_Clock6;
-		case EMLControllerLEDPattern::Clock07:
-			return MLInputControllerFeedbackPatternLED_Clock7;
-		case EMLControllerLEDPattern::Clock08:
-			return MLInputControllerFeedbackPatternLED_Clock8;
-		case EMLControllerLEDPattern::Clock09:
-			return MLInputControllerFeedbackPatternLED_Clock9;
-		case EMLControllerLEDPattern::Clock10:
-			return MLInputControllerFeedbackPatternLED_Clock10;
-		case EMLControllerLEDPattern::Clock11:
-			return MLInputControllerFeedbackPatternLED_Clock11;
-		case EMLControllerLEDPattern::Clock12:
-			return MLInputControllerFeedbackPatternLED_Clock12;
-		case EMLControllerLEDPattern::Clock01_07:
-			return MLInputControllerFeedbackPatternLED_Clock1And7;
-		case EMLControllerLEDPattern::Clock02_08:
-			return MLInputControllerFeedbackPatternLED_Clock2And8;
-		case EMLControllerLEDPattern::Clock03_09:
-			return MLInputControllerFeedbackPatternLED_Clock3And9;
-		case EMLControllerLEDPattern::Clock04_10:
-			return MLInputControllerFeedbackPatternLED_Clock4And10;
-		case EMLControllerLEDPattern::Clock05_11:
-			return MLInputControllerFeedbackPatternLED_Clock5And11;
-		case EMLControllerLEDPattern::Clock06_12:
-			return MLInputControllerFeedbackPatternLED_Clock6And12;    
-		default:
-			UE_LOG(LogMagicLeapController, Error, TEXT("Unhandled LED Pattern type %d"), static_cast<int32>(LEDPattern));
-			break;
+			MessageHandler->OnControllerAnalog(MLTriggerToUnrealTriggerAxis(Hand),
+				DeviceIndex, CurrControllerState->TriggerAnalog);
+		}
 	}
-	return MLInputControllerFeedbackPatternLED_Ensure32Bits;
+#endif //WITH_MLSDK
 }
 
-MLInputControllerFeedbackEffectLED FMagicLeapController::UnrealToMLEffectLED(EMLControllerLEDEffect LEDEffect) const
+EMLControllerType FMagicLeapController::GetMLControllerType(EControllerHand Hand) const
 {
-	switch (LEDEffect)
+#if WITH_MLSDK
+	const auto& MotionSource = ControllerMapper.GetMotionSourceForHand(Hand);
+	if (MotionSource == FMagicLeapMotionSourceNames::Control0)
 	{
-		case EMLControllerLEDEffect::RotateCW:
-			return MLInputControllerFeedbackEffectLED_RotateCW;
-		case EMLControllerLEDEffect::RotateCCW:
-			return MLInputControllerFeedbackEffectLED_RotateCCW;
-		case EMLControllerLEDEffect::Pulse:
-			return MLInputControllerFeedbackEffectLED_Pulse;
-		case EMLControllerLEDEffect::PaintCW:
-			return MLInputControllerFeedbackEffectLED_PaintCW;
-		case EMLControllerLEDEffect::PaintCCW:
-			return MLInputControllerFeedbackEffectLED_PaintCCW;
-		case EMLControllerLEDEffect::Blink:
-			return MLInputControllerFeedbackEffectLED_Blink;
-		default:
-			UE_LOG(LogMagicLeapController, Error, TEXT("Unhandled LED effect type %d"), static_cast<int32>(LEDEffect));
-		break;
+		return EMLControllerType::Device;
 	}
-	return MLInputControllerFeedbackEffectLED_Ensure32Bits;
-}
-
-#define LED_COLOR_CASE(x) case EMLControllerLEDColor::x: { return MLInputControllerFeedbackColorLED_##x; }
-MLInputControllerFeedbackColorLED FMagicLeapController::UnrealToMLColorLED(EMLControllerLEDColor LEDColor) const
-{
-	switch (LEDColor)
+	if (MotionSource == FMagicLeapMotionSourceNames::Control1)
 	{
-#if MLSDK_VERSION_MINOR >= 16
-		LED_COLOR_CASE(BrightMissionRed)
-		LED_COLOR_CASE(PastelMissionRed)
-#else
-		case EMLControllerLEDColor::BrightMissionRed: { return MLInputControllerFeedbackColorLED_BrightRed; }
-		case EMLControllerLEDColor::PastelMissionRed: { return MLInputControllerFeedbackColorLED_PastelRed; }
-#endif // MLSDK_VERSION_MINOR >= 16
-		LED_COLOR_CASE(BrightFloridaOrange)
-		LED_COLOR_CASE(PastelFloridaOrange)
-		LED_COLOR_CASE(BrightLunaYellow)
-		LED_COLOR_CASE(PastelLunaYellow)
-		LED_COLOR_CASE(BrightNebulaPink)
-		LED_COLOR_CASE(PastelNebulaPink)
-		LED_COLOR_CASE(BrightCosmicPurple)
-		LED_COLOR_CASE(PastelCosmicPurple)
-		LED_COLOR_CASE(BrightMysticBlue)
-		LED_COLOR_CASE(PastelMysticBlue)
-		LED_COLOR_CASE(BrightCelestialBlue)
-		LED_COLOR_CASE(PastelCelestialBlue)
-		LED_COLOR_CASE(BrightShaggleGreen)
-		LED_COLOR_CASE(PastelShaggleGreen)
-		default:
-			UE_LOG(LogMagicLeapController, Error, TEXT("Unhandled LED color type %d"), static_cast<int32>(LEDColor));
-			break;
+		return EMLControllerType::Device;
 	}
-	return MLInputControllerFeedbackColorLED_Ensure32Bits;
-}
-
-MLInputControllerFeedbackEffectSpeedLED FMagicLeapController::UnrealToMLSpeedLED(EMLControllerLEDSpeed LEDSpeed) const
-{
-	switch (LEDSpeed)
+	if (MotionSource == FMagicLeapMotionSourceNames::MobileApp)
 	{
-		case EMLControllerLEDSpeed::Slow:
-			return MLInputControllerFeedbackEffectSpeedLED_Slow;
-		case EMLControllerLEDSpeed::Medium:
-			return MLInputControllerFeedbackEffectSpeedLED_Medium;
-		case EMLControllerLEDSpeed::Fast:
-			return MLInputControllerFeedbackEffectSpeedLED_Fast;
-		default:
-			UE_LOG(LogMagicLeapController, Error, TEXT("Unhandled LED speed type %d"), static_cast<int32>(LEDSpeed));
-			break;
+		return EMLControllerType::MobileApp;
 	}
-	return MLInputControllerFeedbackEffectSpeedLED_Ensure32Bits;  
+#endif //WITH_MLSDK
+	return EMLControllerType::None;
 }
 
-MLInputControllerFeedbackPatternVibe FMagicLeapController::UnrealToMLPatternVibe(EMLControllerHapticPattern HapticPattern) const
+bool FMagicLeapController::PlayControllerLED(EControllerHand Hand, EMLControllerLEDPattern LEDPattern, EMLControllerLEDColor LEDColor, float DurationInSec)
 {
-	switch (HapticPattern)
-	{
-		case EMLControllerHapticPattern::None:
-			return MLInputControllerFeedbackPatternVibe_None;
-		case EMLControllerHapticPattern::Click:
-			return MLInputControllerFeedbackPatternVibe_Click;
-		case EMLControllerHapticPattern::Bump:
-			return MLInputControllerFeedbackPatternVibe_Bump;
-		case EMLControllerHapticPattern::DoubleClick:
-			return MLInputControllerFeedbackPatternVibe_DoubleClick;
-		case EMLControllerHapticPattern::Buzz:
-			return MLInputControllerFeedbackPatternVibe_Buzz;
-		case EMLControllerHapticPattern::Tick:
-			return MLInputControllerFeedbackPatternVibe_Tick;
-		case EMLControllerHapticPattern::ForceDown:
-			return MLInputControllerFeedbackPatternVibe_ForceDown;
-		case EMLControllerHapticPattern::ForceUp:
-			return MLInputControllerFeedbackPatternVibe_ForceUp;
-		case EMLControllerHapticPattern::ForceDwell:
-			return MLInputControllerFeedbackPatternVibe_ForceDwell;
-		case EMLControllerHapticPattern::SecondForceDown:
-			return MLInputControllerFeedbackPatternVibe_SecondForceDown;
-		default:
-			UE_LOG(LogMagicLeapController, Error, TEXT("Unhandled Haptic Pattern type %d"), static_cast<int32>(HapticPattern));
-			break;
-	}
-	return MLInputControllerFeedbackPatternVibe_Ensure32Bits;
+#if WITH_MLSDK
+	return PlayLEDPattern(ControllerMapper.GetMotionSourceForHand(Hand), LEDPattern, LEDColor, DurationInSec);
+#endif //WITH_MLSDK
+	return false;
 }
 
-MLInputControllerFeedbackIntensity FMagicLeapController::UnrealToMLHapticIntensity(EMLControllerHapticIntensity HapticIntensity) const
+bool FMagicLeapController::PlayControllerLEDEffect(EControllerHand Hand, EMLControllerLEDEffect LEDEffect, EMLControllerLEDSpeed LEDSpeed, EMLControllerLEDPattern LEDPattern, EMLControllerLEDColor LEDColor, float DurationInSec)
 {
-	switch (HapticIntensity)
-	{
-	case EMLControllerHapticIntensity::Low:
-		return MLInputControllerFeedbackIntensity_Low;
-	case EMLControllerHapticIntensity::Medium:
-		return MLInputControllerFeedbackIntensity_Medium;
-	case EMLControllerHapticIntensity::High:
-		return MLInputControllerFeedbackIntensity_High;
-	default:
-		UE_LOG(LogMagicLeapController, Error, TEXT("Unhandled Haptic Intensity type %d"), static_cast<int32>(HapticIntensity));
-		break;
-	}
-	return MLInputControllerFeedbackIntensity_Ensure32Bits;
-}
+#if WITH_MLSDK
+	return PlayLEDEffect(ControllerMapper.GetMotionSourceForHand(Hand), LEDEffect, LEDSpeed, LEDPattern, LEDColor, DurationInSec);
 #endif //WITH_MLSDK
 
-void FMagicLeapController::InitializeInputCallbacks()
+	return false;
+}
+
+bool FMagicLeapController::PlayControllerHapticFeedback(EControllerHand Hand, EMLControllerHapticPattern HapticPattern, EMLControllerHapticIntensity Intensity)
 {
 #if WITH_MLSDK
-	FMemory::Memset(&InputKeyboardCallbacks, 0, sizeof(InputKeyboardCallbacks));
-
-	InputKeyboardCallbacks.on_char = [](uint32_t char_utf32, void *data)
-	{
-		auto controller = reinterpret_cast<FMagicLeapController*>(data);
-		if (controller)
-		{
-			{
-				FScopeLock Lock(&controller->KeyCriticalSection);
-				controller->PendingCharKeys.Add(char_utf32);
-			}
-		}
-	};
-
-	InputKeyboardCallbacks.on_key_down = [](MLKeyCode key_code, uint32 modifier_mask, void *data)
-	{
-		auto controller = reinterpret_cast<FMagicLeapController*>(data);
-		if (controller)
-		{
-			{
-				FScopeLock Lock(&controller->KeyCriticalSection);
-				controller->PendingKeyDowns.Add(key_code);
-			}
-		}
-	};
-
-	InputKeyboardCallbacks.on_key_up = [](MLKeyCode key_code, uint32 modifier_mask, void *data)
-	{
-		auto controller = reinterpret_cast<FMagicLeapController*>(data);
-		if (controller)
-		{
-			{
-				FScopeLock Lock(&controller->KeyCriticalSection);
-				controller->PendingKeyUps.Add(key_code);
-			}
-		}
-	};
-
-	FMemory::Memset(&InputControllerCallbacks, 0, sizeof(InputControllerCallbacks));
-	// Creating an async task to be fired on the game thread in these callbacks was causing
-	// intermitted blocks on the game thread. Until that is investigated and fixed, use polling.
-	InputControllerCallbacks.on_touchpad_gesture_start = [](uint8 controller_id, const MLInputControllerTouchpadGesture *touchpad_gesture, void *data)
-	{
-		auto controller = reinterpret_cast<FMagicLeapController*>(data);
-		if (controller)
-		{
-			EControllerHand hand;
-			controller->GetControllerMapping(controller_id, hand);
-			FMagicLeapTouchpadGesture gesture = MLToUnrealTouchpadGesture(hand, *touchpad_gesture);
-			for (IMagicLeapTouchpadGestures* Receiver : controller->TouchpadGestureReceivers)
-			{
-				// NOTE: This is temporary; Epic has a Jira task to pipe touchpad gestures through correctly but we need
-				// something for the interim - rmobbs
-				Receiver->OnTouchpadGestureStartCallback(gesture);
-			}
-		}
-	};
-
-	InputControllerCallbacks.on_touchpad_gesture_continue = [](uint8 controller_id, const MLInputControllerTouchpadGesture *touchpad_gesture, void *data)
-	{
-		auto controller = reinterpret_cast<FMagicLeapController*>(data);
-		if (controller)
-		{
-			EControllerHand hand;
-			controller->GetControllerMapping(controller_id, hand);
-			FMagicLeapTouchpadGesture gesture = MLToUnrealTouchpadGesture(hand, *touchpad_gesture);
-			for (IMagicLeapTouchpadGestures* Receiver : controller->TouchpadGestureReceivers)
-			{
-				// NOTE: This is temporary; Epic has a Jira task to pipe touchpad gestures through correctly but we need
-				// something for the interim - rmobbs
-				Receiver->OnTouchpadGestureContinueCallback(gesture);
-			}
-		}
-	};
-
-	InputControllerCallbacks.on_touchpad_gesture_end = [](uint8 controller_id, const MLInputControllerTouchpadGesture *touchpad_gesture, void *data)
-	{
-		auto controller = reinterpret_cast<FMagicLeapController*>(data);
-		if (controller)
-		{
-			EControllerHand hand;
-			controller->GetControllerMapping(controller_id, hand);
-			FMagicLeapTouchpadGesture gesture = MLToUnrealTouchpadGesture(hand, *touchpad_gesture);
-			for (IMagicLeapTouchpadGestures* Receiver : controller->TouchpadGestureReceivers)
-			{
-				// NOTE: This is temporary; Epic has a Jira task to pipe touchpad gestures through correctly but we need
-				// something for the interim - rmobbs
-				Receiver->OnTouchpadGestureEndCallback(gesture);
-			}
-		}
-	};
-
-	InputControllerCallbacks.on_button_down = [](uint8_t controller_id, MLInputControllerButton button, void *data)
-	{
-		auto controller = reinterpret_cast<FMagicLeapController*>(data);
-		if (controller)
-		{
-			FScopeLock Lock(&controller->ButtonCriticalSection);
-			FButtonMap ButtonMap;
-			ButtonMap.ControllerID = controller_id;
-			ButtonMap.Button = button;
-			ButtonMap.bPressed = true;
-			controller->PendingButtonStates.Add(ButtonMap);
-		}
-	};
-
-	InputControllerCallbacks.on_button_up = [](uint8_t controller_id, MLInputControllerButton button, void *data)
-	{
-		auto controller = reinterpret_cast<FMagicLeapController*>(data);
-		if (controller)
-		{
-			FScopeLock Lock(&controller->ButtonCriticalSection);
-			FButtonMap ButtonMap;
-			ButtonMap.ControllerID = controller_id;
-			ButtonMap.Button = button;
-			ButtonMap.bPressed = false;
-			controller->PendingButtonStates.Add(ButtonMap);
-		}
-	};
-
-	InputControllerCallbacks.on_connect = nullptr;
-	InputControllerCallbacks.on_disconnect = nullptr;
+	return PlayHapticPattern(ControllerMapper.GetMotionSourceForHand(Hand), HapticPattern, Intensity);
 #endif //WITH_MLSDK
+
+	return false;
 }
 
 #undef LOCTEXT_NAMESPACE
