Updated ML integration to MLSDK 0.19

diff --git a/Engine/Plugins/Lumin/MagicLeap/Source/MagicLeapController/Private/MagicLeapController.h b/Engine/Plugins/Lumin/MagicLeap/Source/MagicLeapController/Private/MagicLeapController.h
--- a/Engine/Plugins/Lumin/MagicLeap/Source/MagicLeapController/Private/MagicLeapController.h
+++ b/Engine/Plugins/Lumin/MagicLeap/Source/MagicLeapController/Private/MagicLeapController.h
@@ -5,10 +5,10 @@
 #include "GenericPlatform/IInputInterface.h"
 #include "IInputDevice.h"
 #include "IMagicLeapInputDevice.h"
-#include "IHapticDevice.h"
 #include "IMotionController.h"
 #include "XRMotionControllerBase.h"
 #include "IMagicLeapControllerPlugin.h"
+#include <Containers/Queue.h>
 #include "Misc/ScopeLock.h"
 #include "MagicLeapPluginUtil.h" // for ML_INCLUDES_START/END
 
@@ -33,92 +33,95 @@ class IMagicLeapTouchpadGestures;
 /**
  * MagicLeap Motion Controller
  */
-class FMagicLeapController : public IMagicLeapInputDevice, public IHapticDevice, public FXRMotionControllerBase
+class FMagicLeapController : public IMagicLeapInputDevice, public FXRMotionControllerBase
 {
+private:
+#if WITH_MLSDK
+	class FControllerMapper
+	{
+	friend class FMagicLeapController;
+	private:
+		TMap<FName, int32> MotionSourceToInputControllerIndex;
+		FName InputControllerIndexToMotionSource[MLInput_MaxControllers];
+		TMap<EControllerHand, FName> HandToMotionSource;
+		TMap<FName, EControllerHand> MotionSourceToHand;
+		FCriticalSection CriticalSection;
+		EControllerHand DefaultInputControllerIndexToHand[MLInput_MaxControllers];
+
+	protected:
+		void UpdateMotionSourceInputIndexPairing(const MLInputControllerState ControllerState[MLInput_MaxControllers]);
+	public:
+		FControllerMapper();
+
+		void MapHandToMotionSource(EControllerHand Hand, FName MotionSource);
+
+		FName GetMotionSourceForHand(EControllerHand Hand) const;
+		EControllerHand GetHandForMotionSource(FName MotionSource) const;
+
+		FName GetMotionSourceForInputControllerIndex(uint8 controller_id) const;
+		uint8 GetInputControllerIndexForMotionSource(FName MotionSource) const;
+
+		EControllerHand GetHandForInputControllerIndex(uint8 controller_id) const;
+		uint8 GetInputControllerIndexForHand(EControllerHand Hand) const;
+		EMLControllerType MotionSourceToControllerType(FName InMotionSource);
+
+		void SwapHands();
+	};
+#endif //WITH_MLSDK
 public:
 	FMagicLeapController(const TSharedRef<FGenericApplicationMessageHandler>& InMessageHandler);
 	virtual ~FMagicLeapController();
 
 	/** IMagicLeapInputDevice interface */
-	virtual void Tick(float DeltaTime) override;
-	virtual void SendControllerEvents() override;
-	virtual void SetMessageHandler(const TSharedRef<FGenericApplicationMessageHandler>& InMessageHandler) override;
-	virtual bool Exec(UWorld* InWorld, const TCHAR* Cmd, FOutputDevice& Ar) override;
-	virtual void SetChannelValue(int32 ControllerId, FForceFeedbackChannelType ChannelType, float Value) override;
-	virtual void SetChannelValues(int32 ControllerId, const FForceFeedbackValues &values) override;
-	virtual class IHapticDevice* GetHapticDevice() override;
-	virtual bool IsGamepadAttached() const override;
-	virtual void Enable() override;
-	virtual bool SupportsExplicitEnable() const override;
-	virtual void Disable() override;
-
-	/** IHapticDevice interface */
-	virtual void SetHapticFeedbackValues(int32 ControllerId, int32 Hand, const FHapticFeedbackValues& Values) override;
-	virtual void GetHapticFrequencyRange(float& MinFrequency, float& MaxFrequency) const override;
-	virtual float GetHapticAmplitudeScale() const override;
-
-private:
-	void InternalSetHapticFeedbackValues(int32 ControllerId, int32 Hand, const FHapticFeedbackValues& Values, bool bFromHapticInterface);
-public:
+	void Tick(float DeltaTime) override;
+	void SendControllerEvents() override;
+	void SetMessageHandler(const TSharedRef<FGenericApplicationMessageHandler>& InMessageHandler) override;
+	bool Exec(UWorld* InWorld, const TCHAR* Cmd, FOutputDevice& Ar) override;
+	bool IsGamepadAttached() const override;
+	void Enable() override;
+	bool SupportsExplicitEnable() const override;
+	void Disable() override;
 
 	/** IMotionController interface */
-	virtual bool GetControllerOrientationAndPosition(const int32 ControllerIndex, const EControllerHand DeviceHand, FRotator& OutOrientation, FVector& OutPosition, float WorldToMetersScale) const override;
-	virtual ETrackingStatus GetControllerTrackingStatus(const int32 ControllerIndex, const EControllerHand DeviceHand) const override;
-	virtual FName GetMotionControllerDeviceTypeName() const override;
+	bool GetControllerOrientationAndPosition(const int32 ControllerIndex, const FName MotionSource, FRotator& OutOrientation, FVector& OutPosition, float WorldToMetersScale) const override;
+	ETrackingStatus GetControllerTrackingStatus(const int32 ControllerIndex, const FName MotionSource) const override;
+	FName GetMotionControllerDeviceTypeName() const override;
+	void EnumerateSources(TArray<FMotionControllerSource>& SourcesOut) const override;
+
+	/** IInputDevice interface */
+	void SetChannelValue(int32 ControllerId, FForceFeedbackChannelType ChannelType, float Value) override { }
+	void SetChannelValues(int32 ControllerId, const FForceFeedbackValues &values) override { }
 
-	/** FMagicLeapController interface */
-	bool IsInputStateValid() const;
+	EMLControllerTrackingMode GetControllerTrackingMode();
+	bool SetControllerTrackingMode(EMLControllerTrackingMode TrackingMode);
+
+	void RegisterTouchpadGestureReceiver(IMagicLeapTouchpadGestures* Receiver);
+	void UnregisterTouchpadGestureReceiver(IMagicLeapTouchpadGestures* Receiver);
+
+	bool PlayLEDPattern(FName MotionSource, EMLControllerLEDPattern LEDPattern, EMLControllerLEDColor LEDColor, float DurationInSec);
+	bool PlayLEDEffect(FName MotionSource, EMLControllerLEDEffect LEDEffect, EMLControllerLEDSpeed LEDSpeed, EMLControllerLEDPattern LEDPattern, EMLControllerLEDColor LEDColor, float DurationInSec);
+	bool PlayHapticPattern(FName MotionSource, EMLControllerHapticPattern HapticPattern, EMLControllerHapticIntensity Intensity);
 
-	void OnChar(uint32 CharUTF32);
 #if WITH_MLSDK
-	void OnKeyDown(MLKeyCode KeyCode);
-	void OnKeyUp(MLKeyCode KeyCode);
-	void OnButtonDown(uint8 ControllerID, MLInputControllerButton Button);
-	void OnButtonUp(uint8 ControllerID, MLInputControllerButton Button);
+	// Has to be public so button functions can use it
+	FControllerMapper ControllerMapper;
 #endif //WITH_MLSDK
-	void OnControllerConnected(EControllerHand Hand);
-	void OnControllerDisconnected(EControllerHand Hand);
-
-	void SetControllerIsConnected(EControllerHand Hand, bool bConnected);
 
-	bool GetControllerMapping(int32 ControllerIndex, EControllerHand& Hand) const;
-	void InvertControllerMapping();
+	bool GetControllerOrientationAndPosition(const int32 ControllerIndex, const EControllerHand DeviceHand, FRotator& OutOrientation, FVector& OutPosition, float WorldToMetersScale) const override;
+	ETrackingStatus GetControllerTrackingStatus(const int32 ControllerIndex, const EControllerHand DeviceHand) const override;
 	EMLControllerType GetMLControllerType(EControllerHand Hand) const;
-
 	bool PlayControllerLED(EControllerHand Hand, EMLControllerLEDPattern LEDPattern, EMLControllerLEDColor LEDColor, float DurationInSec);
 	bool PlayControllerLEDEffect(EControllerHand Hand, EMLControllerLEDEffect LEDEffect, EMLControllerLEDSpeed LEDSpeed, EMLControllerLEDPattern LEDPattern, EMLControllerLEDColor LEDColor, float DurationInSec);
 	bool PlayControllerHapticFeedback(EControllerHand Hand, EMLControllerHapticPattern HapticPattern, EMLControllerHapticIntensity Intensity);
-	bool SetControllerTrackingMode(EMLControllerTrackingMode TrackingMode);
-	EMLControllerTrackingMode GetControllerTrackingMode();
-
-	void RegisterTouchpadGestureReceiver(IMagicLeapTouchpadGestures* Receiver);
-	void UnregisterTouchpadGestureReceiver(IMagicLeapTouchpadGestures* Receiver);
 
 private:
 	void UpdateTrackerData();
-	void UpdateControllerTransformFromInputTracker(const class FAppFramework& AppFramework, FTransform& ControllerTransform, EControllerHand ControllerHand);
-#if WITH_MLSDK
-	void UpdateControllerTransformFromControllerTracker(const class FAppFramework& AppFramework, const MLControllerSystemState& ControllerSystemState, FTransform& ControllerTransform, int32 DeviceIndex);
-#endif //WITH_MLSDK
-	void SendControllerEventsForHand(EControllerHand Hand);
+	void UpdateControllerStateFromInputTracker(const class FAppFramework& AppFramework, FName MotionSource);
+	void UpdateControllerStateFromControllerTracker(const class FAppFramework& AppFramework, FName MotionSource);
 	void AddKeys();
-	void DebouncedButtonMessageHandler(bool NewButtonState, bool OldButtonState, const FName& ButtonName);
-#if WITH_MLSDK
-	const FName& MagicLeapButtonToUnrealButton(int32 ControllerID, MLInputControllerButton ml_button);
-#endif //WITH_MLSDK
-	const FName& MagicLeapTouchToUnrealThumbstickAxis(EControllerHand Hand, uint32 TouchIndex);
-	const FName& MagicLeapTouchToUnrealThumbstickButton(EControllerHand Hand);
-	const FName& MagicLeapTriggerToUnrealTriggerAxis(EControllerHand Hand);
-	const FName& MagicLeapTriggerToUnrealTriggerKey(EControllerHand Hand);
-#if WITH_MLSDK
-	MLInputControllerFeedbackPatternLED UnrealToMLPatternLED(EMLControllerLEDPattern LEDPattern) const;
-	MLInputControllerFeedbackEffectLED UnrealToMLEffectLED(EMLControllerLEDEffect LEDEffect) const;
-	MLInputControllerFeedbackColorLED UnrealToMLColorLED(EMLControllerLEDColor LEDColor) const;
-	MLInputControllerFeedbackEffectSpeedLED UnrealToMLSpeedLED(EMLControllerLEDSpeed LEDSpeed) const;
-	MLInputControllerFeedbackPatternVibe UnrealToMLPatternVibe(EMLControllerHapticPattern HapticPattern) const;
-	MLInputControllerFeedbackIntensity UnrealToMLHapticIntensity(EMLControllerHapticIntensity HapticIntensity) const;
-#endif //WITH_MLSDK
+	void ReadConfigParams();
 	void InitializeInputCallbacks();
+	void SendControllerEventsForHand(EControllerHand Hand);
 
 	TSharedPtr<FGenericApplicationMessageHandler> MessageHandler;
 	int32 DeviceIndex;
@@ -126,61 +129,23 @@ private:
 #if WITH_MLSDK
 	MLHandle InputTracker;
 	MLHandle ControllerTracker;
-	MLInputControllerState InputState[MLInput_MaxControllers];
-	MLInputControllerState OldInputState[MLInput_MaxControllers];
-	MLInputControllerCallbacks InputControllerCallbacks;
-	MLInputKeyboardCallbacks InputKeyboardCallbacks;
+	MLInputControllerDof ControllerDof;
 	EMLControllerTrackingMode TrackingMode;
+	MLInputControllerState InputControllerState[MLInput_MaxControllers];
+	MLControllerSystemState ControllerSystemState;
+	MLInputControllerCallbacks InputControllerCallbacks;
+	TMap<FName, FMagicLeapControllerState> CurrMotionSourceControllerState;
+	TMap<FName, FMagicLeapControllerState> PrevMotionSourceControllerState;
 #endif //WITH_MLSDK
 
 	bool bIsInputStateValid;
-	bool bTriggerState;
-	bool bTriggerKeyPressing;
-
-	float triggerKeyIsConsideredPressed;
-	float triggerKeyIsConsideredReleased;
 
-	FTransform LeftControllerTransform;
-	FTransform RightControllerTransform;
-
-	TMap<int32, EControllerHand> ControllerIDToHand;
-	TMap<EControllerHand, int32> HandToControllerID;
-
-	TMap<int32, FMagicLeapControllerState> ControllerIDToControllerState;
-
-	struct FMagicLeapKeys
-	{
-		static const FKey MotionController_Left_Thumbstick_Z;
-		static const FKey Left_MoveButton;
-		static const FKey Left_AppButton;
-		static const FKey Left_HomeButton;
-
-		static const FKey MotionController_Right_Thumbstick_Z;
-		static const FKey Right_MoveButton;
-		static const FKey Right_AppButton;
-		static const FKey Right_HomeButton;
-	};
+	float TriggerKeyIsConsideredPressed;
+	float TriggerKeyIsConsideredReleased;
 
 	TArray<IMagicLeapTouchpadGestures*> TouchpadGestureReceivers;
 
-#if WITH_MLSDK
-	TArray<MLKeyCode> PendingKeyDowns;
-	TArray<MLKeyCode> PendingKeyUps;
-#endif //WITH_MLSDK
-	TArray<uint32> PendingCharKeys;
-	FCriticalSection KeyCriticalSection;
-
-	struct FButtonMap
-	{
-#if WITH_MLSDK
-		MLInputControllerButton Button;
-#endif //WITH_MLSDK
-		uint8 ControllerID;
-		bool bPressed;
-	};
-
-	TArray<FButtonMap> PendingButtonStates;
-	FCriticalSection ButtonCriticalSection;
+	TQueue<TPair<FName, bool>> PendingButtonEvents;
 };
 
 DEFINE_LOG_CATEGORY_STATIC(LogMagicLeapController, Display, All);
diff --git a/Engine/Plugins/Lumin/MagicLeap/Source/MagicLeapController/Private/MagicLeapControllerFunctionLibrary.cpp b/Engine/Plugins/Lumin/MagicLeap/Source/MagicLeapController/Private/MagicLeapControllerFunctionLibrary.cpp
--- a/Engine/Plugins/Lumin/MagicLeap/Source/MagicLeapController/Private/MagicLeapControllerFunctionLibrary.cpp
+++ b/Engine/Plugins/Lumin/MagicLeap/Source/MagicLeapController/Private/MagicLeapControllerFunctionLibrary.cpp
@@ -5,91 +5,181 @@
 #include "MagicLeapController.h"
 #include "IMagicLeapPlugin.h"
 
-int32 UMagicLeapControllerFunctionLibrary::MaxSupportedMagicLeapControllers()
+bool UMagicLeapControllerFunctionLibrary::PlayLEDPattern(FName MotionSource, EMLControllerLEDPattern LEDPattern, EMLControllerLEDColor LEDColor, float DurationInSec)
 {
-#if WITH_MLSDK
-	return MLInput_MaxControllers;
-#else
-	return 0;
-#endif //WITH_MLSDK
+	TSharedPtr<FMagicLeapController> controller = StaticCastSharedPtr<FMagicLeapController>(IMagicLeapControllerPlugin::Get().GetInputDevice());
+	if (controller.IsValid())
+	{
+		return controller->PlayLEDPattern(MotionSource, LEDPattern, LEDColor, DurationInSec);
+	}
+	return false;
 }
 
-bool UMagicLeapControllerFunctionLibrary::GetControllerMapping(int32 ControllerIndex, EControllerHand& Hand)
+bool UMagicLeapControllerFunctionLibrary::PlayLEDEffect(FName MotionSource, EMLControllerLEDEffect LEDEffect, EMLControllerLEDSpeed LEDSpeed, EMLControllerLEDPattern LEDPattern, EMLControllerLEDColor LEDColor, float DurationInSec)
 {
 	TSharedPtr<FMagicLeapController> controller = StaticCastSharedPtr<FMagicLeapController>(IMagicLeapControllerPlugin::Get().GetInputDevice());
 	if (controller.IsValid())
 	{
-		return controller->GetControllerMapping(ControllerIndex, Hand);
+		return controller->PlayLEDEffect(MotionSource, LEDEffect, LEDSpeed, LEDPattern, LEDColor, DurationInSec);
 	}
 	return false;
 }
 
-void UMagicLeapControllerFunctionLibrary::InvertControllerMapping()
+bool UMagicLeapControllerFunctionLibrary::PlayHapticPattern(FName MotionSource, EMLControllerHapticPattern HapticPattern, EMLControllerHapticIntensity Intensity)
 {
 	TSharedPtr<FMagicLeapController> controller = StaticCastSharedPtr<FMagicLeapController>(IMagicLeapControllerPlugin::Get().GetInputDevice());
 	if (controller.IsValid())
 	{
-		controller->InvertControllerMapping();
+		return controller->PlayHapticPattern(MotionSource, HapticPattern, Intensity);
 	}
+	return false;
 }
 
-EMLControllerType UMagicLeapControllerFunctionLibrary::GetMLControllerType(EControllerHand Hand)
+bool UMagicLeapControllerFunctionLibrary::SetControllerTrackingMode(EMLControllerTrackingMode TrackingMode)
 {
 	TSharedPtr<FMagicLeapController> controller = StaticCastSharedPtr<FMagicLeapController>(IMagicLeapControllerPlugin::Get().GetInputDevice());
 	if (controller.IsValid())
 	{
-		return controller->GetMLControllerType(Hand);
+		return controller->SetControllerTrackingMode(TrackingMode);
 	}
-	return EMLControllerType::None;
+	return false;
 }
 
-bool UMagicLeapControllerFunctionLibrary::PlayControllerLED(EControllerHand Hand, EMLControllerLEDPattern LEDPattern, EMLControllerLEDColor LEDColor, float DurationInSec)
+EMLControllerTrackingMode UMagicLeapControllerFunctionLibrary::GetControllerTrackingMode()
 {
 	TSharedPtr<FMagicLeapController> controller = StaticCastSharedPtr<FMagicLeapController>(IMagicLeapControllerPlugin::Get().GetInputDevice());
 	if (controller.IsValid())
 	{
-		return controller->PlayControllerLED(Hand, LEDPattern, LEDColor, DurationInSec);
+		return controller->GetControllerTrackingMode();
 	}
+	return EMLControllerTrackingMode::InputService;
+}
+
+FName UMagicLeapControllerFunctionLibrary::GetMotionSourceForHand(EControllerHand Hand)
+{
+#if WITH_MLSDK
+	TSharedPtr<FMagicLeapController> controller = StaticCastSharedPtr<FMagicLeapController>(IMagicLeapControllerPlugin::Get().GetInputDevice());
+	if (controller.IsValid())
+	{
+		return controller->ControllerMapper.GetMotionSourceForHand(Hand);
+	}
+#endif //WITH_MLSDK
+	return FMagicLeapMotionSourceNames::Unknown;
+}
+
+EControllerHand UMagicLeapControllerFunctionLibrary::GetHandForMotionSource(FName MotionSource)
+{
+#if WITH_MLSDK
+	TSharedPtr<FMagicLeapController> controller = StaticCastSharedPtr<FMagicLeapController>(IMagicLeapControllerPlugin::Get().GetInputDevice());
+	if (controller.IsValid())
+	{
+		return controller->ControllerMapper.GetHandForMotionSource(MotionSource);
+	}
+#endif //WITH_MLSDK
+	return EControllerHand::ControllerHand_Count;
+}
+
+bool UMagicLeapControllerFunctionLibrary::SetMotionSourceForHand(EControllerHand Hand, FName MotionSource)
+{
+#if WITH_MLSDK
+	TSharedPtr<FMagicLeapController> controller = StaticCastSharedPtr<FMagicLeapController>(IMagicLeapControllerPlugin::Get().GetInputDevice());
+	if (controller.IsValid())
+	{
+		controller->ControllerMapper.MapHandToMotionSource(Hand, MotionSource);
+		return true;
+	}
+#endif //WITH_MLSDK
 	return false;
 }
 
-bool UMagicLeapControllerFunctionLibrary::PlayControllerLEDEffect(EControllerHand Hand, EMLControllerLEDEffect LEDEffect, EMLControllerLEDSpeed LEDSpeed, EMLControllerLEDPattern LEDPattern, EMLControllerLEDColor LEDColor, float DurationInSec)
+EMLControllerType UMagicLeapControllerFunctionLibrary::GetControllerType(FName MotionSource)
 {
+#if WITH_MLSDK
 	TSharedPtr<FMagicLeapController> controller = StaticCastSharedPtr<FMagicLeapController>(IMagicLeapControllerPlugin::Get().GetInputDevice());
 	if (controller.IsValid())
 	{
-		return controller->PlayControllerLEDEffect(Hand, LEDEffect, LEDSpeed, LEDPattern, LEDColor, DurationInSec);
+		return controller->ControllerMapper.MotionSourceToControllerType(MotionSource);
 	}
+#endif //WITH_MLSDK
+	return EMLControllerType::None;
+}
+
+
+
+
+
+
+
+
+
+
+
+
+
+/////////////////////////////////////////////////////////////////////////////////////////////
+// DEPRECATED FUNCTIONS
+/////////////////////////////////////////////////////////////////////////////////////////////
+int32 UMagicLeapControllerFunctionLibrary::MaxSupportedMagicLeapControllers()
+{
+#if WITH_MLSDK
+	return MLInput_MaxControllers;
+#else
+	return 0;
+#endif //WITH_MLSDK
+}
+
+bool UMagicLeapControllerFunctionLibrary::GetControllerMapping(int32 ControllerIndex, EControllerHand& Hand)
+{
 	return false;
 }
 
-bool UMagicLeapControllerFunctionLibrary::PlayControllerHapticFeedback(EControllerHand Hand, EMLControllerHapticPattern HapticPattern, EMLControllerHapticIntensity Intensity)
+void UMagicLeapControllerFunctionLibrary::InvertControllerMapping()
+{
+#if WITH_MLSDK
+	TSharedPtr<FMagicLeapController> controller = StaticCastSharedPtr<FMagicLeapController>(IMagicLeapControllerPlugin::Get().GetInputDevice());
+	if (controller.IsValid())
+	{
+		return controller->ControllerMapper.SwapHands();
+	}
+#endif //WITH_MLSDK
+}
+
+EMLControllerType UMagicLeapControllerFunctionLibrary::GetMLControllerType(EControllerHand Hand)
 {
 	TSharedPtr<FMagicLeapController> controller = StaticCastSharedPtr<FMagicLeapController>(IMagicLeapControllerPlugin::Get().GetInputDevice());
 	if (controller.IsValid())
 	{
-		return controller->PlayControllerHapticFeedback(Hand, HapticPattern, Intensity);
+		return controller->GetMLControllerType(Hand);
+	}
+	return EMLControllerType::None;
+}
+
+bool UMagicLeapControllerFunctionLibrary::PlayControllerLED(EControllerHand Hand, EMLControllerLEDPattern LEDPattern, EMLControllerLEDColor LEDColor, float DurationInSec)
+{
+	TSharedPtr<FMagicLeapController> controller = StaticCastSharedPtr<FMagicLeapController>(IMagicLeapControllerPlugin::Get().GetInputDevice());
+	if (controller.IsValid())
+	{
+		return controller->PlayControllerLED(Hand, LEDPattern, LEDColor, DurationInSec);
 	}
 	return false;
 }
 
-bool UMagicLeapControllerFunctionLibrary::SetControllerTrackingMode(EMLControllerTrackingMode TrackingMode)
+bool UMagicLeapControllerFunctionLibrary::PlayControllerLEDEffect(EControllerHand Hand, EMLControllerLEDEffect LEDEffect, EMLControllerLEDSpeed LEDSpeed, EMLControllerLEDPattern LEDPattern, EMLControllerLEDColor LEDColor, float DurationInSec)
 {
 	TSharedPtr<FMagicLeapController> Controller = StaticCastSharedPtr<FMagicLeapController>(IMagicLeapControllerPlugin::Get().GetInputDevice());
 	if (Controller.IsValid())
 	{
-		return Controller->SetControllerTrackingMode(TrackingMode);
+		return Controller->PlayControllerLEDEffect(Hand, LEDEffect, LEDSpeed, LEDPattern, LEDColor, DurationInSec);
 	}
 	return false;
 }
 
-EMLControllerTrackingMode UMagicLeapControllerFunctionLibrary::GetControllerTrackingMode()
+bool UMagicLeapControllerFunctionLibrary::PlayControllerHapticFeedback(EControllerHand Hand, EMLControllerHapticPattern HapticPattern, EMLControllerHapticIntensity Intensity)
 {
 	TSharedPtr<FMagicLeapController> Controller = StaticCastSharedPtr<FMagicLeapController>(IMagicLeapControllerPlugin::Get().GetInputDevice());
 	if (Controller.IsValid())
 	{
-		return Controller->GetControllerTrackingMode();
+		return Controller->PlayControllerHapticFeedback(Hand, HapticPattern, Intensity);
 	}
-	return EMLControllerTrackingMode::InputService;
+	return false;
 }
-
diff --git a/Engine/Plugins/Lumin/MagicLeap/Source/MagicLeapController/Private/MagicLeapControllerKeys.cpp b/Engine/Plugins/Lumin/MagicLeap/Source/MagicLeapController/Private/MagicLeapControllerKeys.cpp
new file mode 100644
--- /dev/null
+++ b/Engine/Plugins/Lumin/MagicLeap/Source/MagicLeapController/Private/MagicLeapControllerKeys.cpp
@@ -0,0 +1,30 @@
+#include "MagicLeapControllerKeys.h"
+
+const FName FMagicLeapMotionSourceNames::Control0(TEXT("MagicLeapControl0"));
+const FName FMagicLeapMotionSourceNames::Control1(TEXT("MagicLeapControl1"));
+const FName FMagicLeapMotionSourceNames::MobileApp(TEXT("MagicLeapMobileApp"));
+const FName FMagicLeapMotionSourceNames::Unknown(TEXT("Unknown"));
+
+const FKey FMagicLeapKeys::MotionController_Left_Thumbstick_Z("MotionController_Left_Thumbstick_Z");
+const FName FMagicLeapControllerKeyNames::MotionController_Left_Thumbstick_Z_Name("MotionController_Left_Thumbstick_Z");
+
+const FKey FMagicLeapKeys::Left_MoveButton("MagicLeap_Left_MoveButton");
+const FName FMagicLeapControllerKeyNames::Left_MoveButton_Name("MagicLeap_Left_MoveButton");
+
+const FKey FMagicLeapKeys::Left_AppButton("MagicLeap_Left_AppButton");
+const FName FMagicLeapControllerKeyNames::Left_AppButton_Name("MagicLeap_Left_AppButton");
+
+const FKey FMagicLeapKeys::Left_HomeButton("MagicLeap_Left_HomeButton");
+const FName FMagicLeapControllerKeyNames::Left_HomeButton_Name("MagicLeap_Left_HomeButton");
+
+const FKey FMagicLeapKeys::MotionController_Right_Thumbstick_Z("MotionController_Right_Thumbstick_Z");
+const FName FMagicLeapControllerKeyNames::MotionController_Right_Thumbstick_Z_Name("MotionController_Right_Thumbstick_Z");
+
+const FKey FMagicLeapKeys::Right_MoveButton("MagicLeap_Right_MoveButton");
+const FName FMagicLeapControllerKeyNames::Right_MoveButton_Name("MagicLeap_Right_MoveButton");
+
+const FKey FMagicLeapKeys::Right_AppButton("MagicLeap_Right_AppButton");
+const FName FMagicLeapControllerKeyNames::Right_AppButton_Name("MagicLeap_Right_AppButton");
+
+const FKey FMagicLeapKeys::Right_HomeButton("MagicLeap_Right_HomeButton");
+const FName FMagicLeapControllerKeyNames::Right_HomeButton_Name("MagicLeap_Right_HomeButton");
diff --git a/Engine/Plugins/Lumin/MagicLeap/Source/MagicLeapController/Private/MagicLeapControllerMappings.h b/Engine/Plugins/Lumin/MagicLeap/Source/MagicLeapController/Private/MagicLeapControllerMappings.h
new file mode 100644
--- /dev/null
+++ b/Engine/Plugins/Lumin/MagicLeap/Source/MagicLeapController/Private/MagicLeapControllerMappings.h
@@ -0,0 +1,342 @@
+#if WITH_MLSDK
+#define TOUCH_GESTURE_CASE(x) case MLInputControllerTouchpadGestureType_##x: { return EMagicLeapTouchpadGestureType::x; }
+
+EMagicLeapTouchpadGestureType MLToUnrealTouchpadGestureType(MLInputControllerTouchpadGestureType GestureType)
+{
+	switch (GestureType)
+	{
+		TOUCH_GESTURE_CASE(Tap)
+		TOUCH_GESTURE_CASE(ForceTapDown)
+		TOUCH_GESTURE_CASE(ForceTapUp)
+		TOUCH_GESTURE_CASE(ForceDwell)
+		TOUCH_GESTURE_CASE(SecondForceDown)
+		TOUCH_GESTURE_CASE(LongHold)
+		TOUCH_GESTURE_CASE(RadialScroll)
+		TOUCH_GESTURE_CASE(Swipe)
+		TOUCH_GESTURE_CASE(Scroll)
+		TOUCH_GESTURE_CASE(Pinch)
+	}
+	return EMagicLeapTouchpadGestureType::None;
+}
+
+#define TOUCH_DIR_CASE(x) case MLInputControllerTouchpadGestureDirection_##x: { return EMagicLeapTouchpadGestureDirection::x; }
+
+EMagicLeapTouchpadGestureDirection MLToUnrealTouchpadGestureDirection(MLInputControllerTouchpadGestureDirection Direction)
+{
+	switch (Direction)
+	{
+		TOUCH_DIR_CASE(Up)
+		TOUCH_DIR_CASE(Down)
+		TOUCH_DIR_CASE(Left)
+		TOUCH_DIR_CASE(Right)
+		TOUCH_DIR_CASE(In)
+		TOUCH_DIR_CASE(Out)
+		TOUCH_DIR_CASE(Clockwise)
+		TOUCH_DIR_CASE(CounterClockwise)
+	}
+	return EMagicLeapTouchpadGestureDirection::None;
+}
+
+FMagicLeapTouchpadGesture MLToUnrealTouchpadGesture(EControllerHand hand, FName MotionSource, const MLInputControllerTouchpadGesture& touchpad_gesture)
+{
+	FMagicLeapTouchpadGesture gesture;
+	gesture.Hand = hand;
+	gesture.MotionSource = MotionSource;
+	gesture.Type = MLToUnrealTouchpadGestureType(touchpad_gesture.type);
+	gesture.Direction = MLToUnrealTouchpadGestureDirection(touchpad_gesture.direction);
+	gesture.PositionAndForce = FVector(touchpad_gesture.pos_and_force.x, touchpad_gesture.pos_and_force.y, touchpad_gesture.pos_and_force.z);
+	gesture.Speed = touchpad_gesture.speed;
+	gesture.Distance = touchpad_gesture.distance;
+	gesture.FingerGap = touchpad_gesture.finger_gap;
+	gesture.Radius = touchpad_gesture.radius;
+	gesture.Angle = touchpad_gesture.angle;
+
+	return gesture;
+}
+
+const FName& MLToUnrealButton(EControllerHand Hand, MLInputControllerButton ml_button)
+{
+	static const FName empty;
+
+	switch (ml_button)
+	{
+	case MLInputControllerButton_Move:
+		if (Hand == EControllerHand::Left)
+		{
+			return FMagicLeapControllerKeyNames::Left_MoveButton_Name;
+		}
+		return FMagicLeapControllerKeyNames::Right_MoveButton_Name;
+	case MLInputControllerButton_App:
+		if (Hand == EControllerHand::Left)
+		{
+			return FMagicLeapControllerKeyNames::Left_AppButton_Name;
+		}
+		return FMagicLeapControllerKeyNames::Right_AppButton_Name;
+	case MLInputControllerButton_Bumper:
+		if (Hand == EControllerHand::Left)
+		{
+			return FGamepadKeyNames::MotionController_Left_Shoulder;
+		}
+		return FGamepadKeyNames::MotionController_Right_Shoulder;
+	case MLInputControllerButton_HomeTap:
+		if (Hand == EControllerHand::Left)
+		{
+			return FMagicLeapControllerKeyNames::Left_HomeButton_Name;
+		}
+		return FMagicLeapControllerKeyNames::Right_HomeButton_Name;
+	default:
+		break;
+	}
+	return empty;
+}
+
+const FName& MLToUnrealButton(FName MotionSource, MLInputControllerButton ml_button)
+{
+	static const FName empty;
+	return empty;
+}
+#endif //WITH_MLSDK
+
+
+const FName& MLTouchToUnrealThumbstickAxis(EControllerHand Hand, uint32 TouchIndex)
+{
+	static const FName empty;
+
+	switch (TouchIndex)
+	{
+	case 0:
+		if (Hand == EControllerHand::Left)
+		{
+			return FGamepadKeyNames::MotionController_Left_Thumbstick_X;
+		}
+		return FGamepadKeyNames::MotionController_Right_Thumbstick_X;
+	case 1:
+		if (Hand == EControllerHand::Left)
+		{
+			return FGamepadKeyNames::MotionController_Left_Thumbstick_Y;
+		}
+		return FGamepadKeyNames::MotionController_Right_Thumbstick_Y;
+	case 2:
+		if (Hand == EControllerHand::Left)
+		{
+			return FMagicLeapControllerKeyNames::MotionController_Left_Thumbstick_Z_Name;
+		}
+		return FMagicLeapControllerKeyNames::MotionController_Right_Thumbstick_Z_Name;
+	default:
+		return empty;
+	}
+
+	return empty;
+}
+
+const FName& MLTouchToUnrealThumbstickButton(EControllerHand Hand)
+{
+	static const FName empty;
+
+	switch (Hand)
+	{
+	case EControllerHand::Left:
+		return FGamepadKeyNames::MotionController_Left_Thumbstick;
+	case EControllerHand::Right:
+		return FGamepadKeyNames::MotionController_Right_Thumbstick;
+	}
+
+	return empty;
+}
+
+const FName& MLTriggerToUnrealTriggerAxis(EControllerHand Hand)
+{
+	static const FName empty;
+
+	switch (Hand)
+	{
+	case EControllerHand::Left:
+		return FGamepadKeyNames::MotionController_Left_TriggerAxis;
+	case EControllerHand::Right:
+		return FGamepadKeyNames::MotionController_Right_TriggerAxis;
+	}
+
+	return empty;
+}
+
+const FName& MLTriggerToUnrealTriggerKey(EControllerHand Hand)
+{
+	static const FName empty;
+	switch (Hand)
+	{
+	case EControllerHand::Left:
+		return FGamepadKeyNames::MotionController_Left_Trigger;
+	case EControllerHand::Right:
+		return FGamepadKeyNames::MotionController_Right_Trigger;
+	}
+	return empty;
+}
+
+#if WITH_MLSDK
+
+MLInputControllerFeedbackPatternLED UnrealToMLPatternLED(EMLControllerLEDPattern LEDPattern)
+{
+	switch (LEDPattern)
+	{
+	case EMLControllerLEDPattern::None:
+		return MLInputControllerFeedbackPatternLED_None;
+	case EMLControllerLEDPattern::Clock01:
+		return MLInputControllerFeedbackPatternLED_Clock1;
+	case EMLControllerLEDPattern::Clock02:
+		return MLInputControllerFeedbackPatternLED_Clock2;
+	case EMLControllerLEDPattern::Clock03:
+		return MLInputControllerFeedbackPatternLED_Clock3;
+	case EMLControllerLEDPattern::Clock04:
+		return MLInputControllerFeedbackPatternLED_Clock4;
+	case EMLControllerLEDPattern::Clock05:
+		return MLInputControllerFeedbackPatternLED_Clock5;
+	case EMLControllerLEDPattern::Clock06:
+		return MLInputControllerFeedbackPatternLED_Clock6;
+	case EMLControllerLEDPattern::Clock07:
+		return MLInputControllerFeedbackPatternLED_Clock7;
+	case EMLControllerLEDPattern::Clock08:
+		return MLInputControllerFeedbackPatternLED_Clock8;
+	case EMLControllerLEDPattern::Clock09:
+		return MLInputControllerFeedbackPatternLED_Clock9;
+	case EMLControllerLEDPattern::Clock10:
+		return MLInputControllerFeedbackPatternLED_Clock10;
+	case EMLControllerLEDPattern::Clock11:
+		return MLInputControllerFeedbackPatternLED_Clock11;
+	case EMLControllerLEDPattern::Clock12:
+		return MLInputControllerFeedbackPatternLED_Clock12;
+	case EMLControllerLEDPattern::Clock01_07:
+		return MLInputControllerFeedbackPatternLED_Clock1And7;
+	case EMLControllerLEDPattern::Clock02_08:
+		return MLInputControllerFeedbackPatternLED_Clock2And8;
+	case EMLControllerLEDPattern::Clock03_09:
+		return MLInputControllerFeedbackPatternLED_Clock3And9;
+	case EMLControllerLEDPattern::Clock04_10:
+		return MLInputControllerFeedbackPatternLED_Clock4And10;
+	case EMLControllerLEDPattern::Clock05_11:
+		return MLInputControllerFeedbackPatternLED_Clock5And11;
+	case EMLControllerLEDPattern::Clock06_12:
+		return MLInputControllerFeedbackPatternLED_Clock6And12;
+	default:
+		UE_LOG(LogMagicLeapController, Error, TEXT("Unhandled LED Pattern type %d"), static_cast<int32>(LEDPattern));
+		break;
+	}
+	return MLInputControllerFeedbackPatternLED_Ensure32Bits;
+}
+
+MLInputControllerFeedbackEffectLED UnrealToMLEffectLED(EMLControllerLEDEffect LEDEffect)
+{
+	switch (LEDEffect)
+	{
+	case EMLControllerLEDEffect::RotateCW:
+		return MLInputControllerFeedbackEffectLED_RotateCW;
+	case EMLControllerLEDEffect::RotateCCW:
+		return MLInputControllerFeedbackEffectLED_RotateCCW;
+	case EMLControllerLEDEffect::Pulse:
+		return MLInputControllerFeedbackEffectLED_Pulse;
+	case EMLControllerLEDEffect::PaintCW:
+		return MLInputControllerFeedbackEffectLED_PaintCW;
+	case EMLControllerLEDEffect::PaintCCW:
+		return MLInputControllerFeedbackEffectLED_PaintCCW;
+	case EMLControllerLEDEffect::Blink:
+		return MLInputControllerFeedbackEffectLED_Blink;
+	default:
+		UE_LOG(LogMagicLeapController, Error, TEXT("Unhandled LED effect type %d"), static_cast<int32>(LEDEffect));
+		break;
+	}
+	return MLInputControllerFeedbackEffectLED_Ensure32Bits;
+}
+
+#define LED_COLOR_CASE(x) case EMLControllerLEDColor::x: { return MLInputControllerFeedbackColorLED_##x; }
+MLInputControllerFeedbackColorLED UnrealToMLColorLED(EMLControllerLEDColor LEDColor)
+{
+	switch (LEDColor)
+	{
+		LED_COLOR_CASE(BrightMissionRed)
+			LED_COLOR_CASE(PastelMissionRed)
+			LED_COLOR_CASE(BrightFloridaOrange)
+			LED_COLOR_CASE(PastelFloridaOrange)
+			LED_COLOR_CASE(BrightLunaYellow)
+			LED_COLOR_CASE(PastelLunaYellow)
+			LED_COLOR_CASE(BrightNebulaPink)
+			LED_COLOR_CASE(PastelNebulaPink)
+			LED_COLOR_CASE(BrightCosmicPurple)
+			LED_COLOR_CASE(PastelCosmicPurple)
+			LED_COLOR_CASE(BrightMysticBlue)
+			LED_COLOR_CASE(PastelMysticBlue)
+			LED_COLOR_CASE(BrightCelestialBlue)
+			LED_COLOR_CASE(PastelCelestialBlue)
+			LED_COLOR_CASE(BrightShaggleGreen)
+			LED_COLOR_CASE(PastelShaggleGreen)
+	default:
+		UE_LOG(LogMagicLeapController, Error, TEXT("Unhandled LED color type %d"), static_cast<int32>(LEDColor));
+		break;
+	}
+	return MLInputControllerFeedbackColorLED_Ensure32Bits;
+}
+
+MLInputControllerFeedbackEffectSpeedLED UnrealToMLSpeedLED(EMLControllerLEDSpeed LEDSpeed)
+{
+	switch (LEDSpeed)
+	{
+	case EMLControllerLEDSpeed::Slow:
+		return MLInputControllerFeedbackEffectSpeedLED_Slow;
+	case EMLControllerLEDSpeed::Medium:
+		return MLInputControllerFeedbackEffectSpeedLED_Medium;
+	case EMLControllerLEDSpeed::Fast:
+		return MLInputControllerFeedbackEffectSpeedLED_Fast;
+	default:
+		UE_LOG(LogMagicLeapController, Error, TEXT("Unhandled LED speed type %d"), static_cast<int32>(LEDSpeed));
+		break;
+	}
+	return MLInputControllerFeedbackEffectSpeedLED_Ensure32Bits;
+}
+
+MLInputControllerFeedbackPatternVibe UnrealToMLPatternVibe(EMLControllerHapticPattern HapticPattern)
+{
+	switch (HapticPattern)
+	{
+	case EMLControllerHapticPattern::None:
+		return MLInputControllerFeedbackPatternVibe_None;
+	case EMLControllerHapticPattern::Click:
+		return MLInputControllerFeedbackPatternVibe_Click;
+	case EMLControllerHapticPattern::Bump:
+		return MLInputControllerFeedbackPatternVibe_Bump;
+	case EMLControllerHapticPattern::DoubleClick:
+		return MLInputControllerFeedbackPatternVibe_DoubleClick;
+	case EMLControllerHapticPattern::Buzz:
+		return MLInputControllerFeedbackPatternVibe_Buzz;
+	case EMLControllerHapticPattern::Tick:
+		return MLInputControllerFeedbackPatternVibe_Tick;
+	case EMLControllerHapticPattern::ForceDown:
+		return MLInputControllerFeedbackPatternVibe_ForceDown;
+	case EMLControllerHapticPattern::ForceUp:
+		return MLInputControllerFeedbackPatternVibe_ForceUp;
+	case EMLControllerHapticPattern::ForceDwell:
+		return MLInputControllerFeedbackPatternVibe_ForceDwell;
+	case EMLControllerHapticPattern::SecondForceDown:
+		return MLInputControllerFeedbackPatternVibe_SecondForceDown;
+	default:
+		UE_LOG(LogMagicLeapController, Error, TEXT("Unhandled Haptic Pattern type %d"), static_cast<int32>(HapticPattern));
+		break;
+	}
+	return MLInputControllerFeedbackPatternVibe_Ensure32Bits;
+}
+
+MLInputControllerFeedbackIntensity UnrealToMLHapticIntensity(EMLControllerHapticIntensity HapticIntensity)
+{
+	switch (HapticIntensity)
+	{
+	case EMLControllerHapticIntensity::Low:
+		return MLInputControllerFeedbackIntensity_Low;
+	case EMLControllerHapticIntensity::Medium:
+		return MLInputControllerFeedbackIntensity_Medium;
+	case EMLControllerHapticIntensity::High:
+		return MLInputControllerFeedbackIntensity_High;
+	default:
+		UE_LOG(LogMagicLeapController, Error, TEXT("Unhandled Haptic Intensity type %d"), static_cast<int32>(HapticIntensity));
+		break;
+	}
+	return MLInputControllerFeedbackIntensity_Ensure32Bits;
+}
+#endif //WITH_MLSDK
+
diff --git a/Engine/Plugins/Lumin/MagicLeap/Source/MagicLeapController/Private/MagicLeapControllerPlugin.cpp b/Engine/Plugins/Lumin/MagicLeap/Source/MagicLeapController/Private/MagicLeapControllerPlugin.cpp
new file mode 100644
--- /dev/null
+++ b/Engine/Plugins/Lumin/MagicLeap/Source/MagicLeapController/Private/MagicLeapControllerPlugin.cpp
@@ -0,0 +1,73 @@
+#include "IMagicLeapControllerPlugin.h"
+#include "MagicLeapController.h"
+#include "Framework/Application/SlateApplication.h"
+
+class FMagicLeapControllerPlugin : public IMagicLeapControllerPlugin
+{
+public:
+
+	virtual void StartupModule() override
+	{
+		IMagicLeapControllerPlugin::StartupModule();
+
+		// HACK: Generic Application might not be instantiated at this point so we create the input device with a 
+		// dummy message handler. When the Generic Application creates the input device it passes a valid message
+		// handler to it which is further on used for all the controller events. This hack fixes issues caused by
+		// using a custom input device before the Generic Application has instantiated it. Eg. within BeginPlay()
+		//
+		// This also fixes the warnings that pop up on the custom input keys when the blueprint loads. Those
+		// warnings are caused because Unreal loads the bluerints before the input device has been instantiated
+		// and has added its keys, thus leading Unreal to believe that those keys don't exist. This hack causes
+		// an earlier instantiation of the input device, and consequently, the custom keys.
+		TSharedPtr<FGenericApplicationMessageHandler> DummyMessageHandler(new FGenericApplicationMessageHandler());
+		CreateInputDevice(DummyMessageHandler.ToSharedRef());
+
+		// Ideally, we should be able to query GetDefault<UMagicLeapSettings>()->bEnableZI directly.
+		// Unfortunately, the UObject system hasn't finished initialization when this module has been loaded.
+		bool bIsVDZIEnabled = false;
+		GConfig->GetBool(TEXT("/Script/MagicLeap.MagicLeapSettings"), TEXT("bEnableZI"), bIsVDZIEnabled, GEngineIni);
+
+		APISetup.Startup(bIsVDZIEnabled);
+#if WITH_MLSDK
+		APISetup.LoadDLL(TEXT("ml_input"));
+#endif
+	}
+
+	virtual void ShutdownModule() override
+	{
+		APISetup.Shutdown();
+		IMagicLeapControllerPlugin::ShutdownModule();
+	}
+
+	virtual TSharedPtr<class IInputDevice> CreateInputDevice(const TSharedRef<FGenericApplicationMessageHandler>& InMessageHandler) override
+	{
+		if (!InputDevice.IsValid())
+		{
+			TSharedPtr<IInputDevice> MagicLeapController(new FMagicLeapController(InMessageHandler));
+			InputDevice = MagicLeapController;
+			return InputDevice;
+		}
+		else
+		{
+			InputDevice.Get()->SetMessageHandler(InMessageHandler);
+			return InputDevice;
+		}
+		return nullptr;
+	}
+
+	virtual TSharedPtr<IInputDevice> GetInputDevice() override
+	{
+		if (!InputDevice.IsValid())
+		{
+			InputDevice = CreateInputDevice(FSlateApplication::Get().GetPlatformApplication()->GetMessageHandler());
+		}
+		return InputDevice;
+	}
+
+private:
+	FMagicLeapAPISetup APISetup;
+	TSharedPtr<IInputDevice> InputDevice;
+};
+
+IMPLEMENT_MODULE(FMagicLeapControllerPlugin, MagicLeapController);
+
diff --git a/Engine/Plugins/Lumin/MagicLeap/Source/MagicLeapController/Private/MagicLeapInputState.h b/Engine/Plugins/Lumin/MagicLeap/Source/MagicLeapController/Private/MagicLeapInputState.h
--- a/Engine/Plugins/Lumin/MagicLeap/Source/MagicLeapController/Private/MagicLeapInputState.h
+++ b/Engine/Plugins/Lumin/MagicLeap/Source/MagicLeapController/Private/MagicLeapInputState.h
@@ -8,37 +8,15 @@
 
 struct FMagicLeapControllerState
 {
-	/** True if the device is connected, otherwise false */
-	bool bIsConnected;
-
-	/** True if position is being tracked, otherwise false */
-	bool bIsPositionTracked;
-
-	/** True if orientation is being tracked, otherwise false */
-	bool bIsOrientationTracked;
-
-	/** Whether or not we're playing a haptic effect.  If true, force feedback calls will be early-outed in favor of the haptic effect */
-	bool bPlayingHapticEffect;
-	/** Haptic frequency (zero to disable) */
-	float HapticFrequency;
-	/** Haptic amplitude (zero to disable) */
-	float HapticAmplitude;
-
-	/** Explicit constructor sets up sensible defaults */
-	FMagicLeapControllerState(const EControllerHand Hand)
-		: bIsConnected(false),
-		bIsPositionTracked(false),
-		bIsOrientationTracked(false),
-		bPlayingHapticEffect(false),
-		HapticFrequency(0.0f),
-		HapticAmplitude(0.0f)
-	{
-	}
-
-	/** Default constructor does nothing.  Don't use it.  This only exists because we cannot initialize an array of objects with no default constructor on non-C++ 11 compliant compilers (VS 2013) */
-	FMagicLeapControllerState()
-	{
-	}
+	static constexpr int kMaxTouches = 2;
+
+	bool bIsConnected = false;
+	bool bTriggerKeyPressing = false;
+	bool bTouchActive[kMaxTouches] = { };
+	float TriggerAnalog = 0.0f;
+	ETrackingStatus TrackingStatus = ETrackingStatus::NotTracked;
+	FTransform Transform;
+	FVector TouchPosAndForce[kMaxTouches] = { };
 };
 
 
diff --git a/Engine/Plugins/Lumin/MagicLeap/Source/MagicLeapController/Public/MagicLeapControllerFunctionLibrary.h b/Engine/Plugins/Lumin/MagicLeap/Source/MagicLeapController/Public/MagicLeapControllerFunctionLibrary.h
--- a/Engine/Plugins/Lumin/MagicLeap/Source/MagicLeapController/Public/MagicLeapControllerFunctionLibrary.h
+++ b/Engine/Plugins/Lumin/MagicLeap/Source/MagicLeapController/Public/MagicLeapControllerFunctionLibrary.h
@@ -15,12 +15,92 @@ class MAGICLEAPCONTROLLER_API UMagicLeapControllerFunctionLibrary : public UBlue
 	GENERATED_BODY()
 
 public:
+	/**
+	Light up the LED on the Magic Leap Controller in the given pattern for the specified duration.
+	@param MotionSource Controller to play the LED pattern on.
+	@param LEDPattern Pattern to play on the controller.
+	@param LEDColor Color of the LED.
+	@param DurationInSec Duration (in seconds) to play the LED pattern.
+	@return True if the command to play the LED pattern was successfully sent to the controller, false otherwise.
+	*/
+	UFUNCTION(BlueprintCallable, Category = "MotionController|MagicLeap")
+		static bool PlayLEDPattern(FName MotionSource, EMLControllerLEDPattern LEDPattern, EMLControllerLEDColor LEDColor, float DurationInSec);
+
+	/**
+	Starts a LED feedback effect using the specified pattern on the specified controller.
+	@param MotionSource Controller to play the LED pattern on.
+	@param LEDEffect Effect to play on the controller.
+	@param LEDSpeed Effect speed.
+	@param LEDPattern Pattern to play on the controller.
+	@param LEDColor Color of the LED.
+	@param DurationInSec Duration (in seconds) to play the LED pattern.
+	@return True if the command to play the LED effect was successfully sent to the controller, false otherwise.
+	*/
+	UFUNCTION(BlueprintCallable, Category = "MotionController|MagicLeap")
+		static bool PlayLEDEffect(FName MotionSource, EMLControllerLEDEffect LEDEffect, EMLControllerLEDSpeed LEDSpeed, EMLControllerLEDPattern LEDPattern, EMLControllerLEDColor LEDColor, float DurationInSec);
+
+	/**
+	Play haptic feedback on the controller.
+	@param MotionSource Controller to play the haptic feedback on.
+	@param HapticPattern Pattern to play on the controller.
+	@param Intensity Intensity to play on the controller.
+	@return True if the command to play the haptic feedback was successfully sent to the controller, false otherwise.
+	*/
+	UFUNCTION(BlueprintCallable, Category = "MotionController|MagicLeap")
+		static bool PlayHapticPattern(FName MotionSource, EMLControllerHapticPattern HapticPattern, EMLControllerHapticIntensity Intensity);
+
+	/**
+	Set controller tracking mode.
+	@param TrackingMode Controller tracking mode. Note that this is global (it affects all control devices).
+	@return True if the command to set the tracking mode was successfully sent to the controller, false otherwise.
+	*/
+	UFUNCTION(BlueprintCallable, Category = "MotionController|MagicLeap")
+		static bool SetControllerTrackingMode(EMLControllerTrackingMode TrackingMode);
+
+	/**
+	Get controller tracking mode.
+	@return Controller tracking mode.
+	*/
+	UFUNCTION(BlueprintCallable, Category = "MotionController|MagicLeap")
+		static EMLControllerTrackingMode GetControllerTrackingMode();
+
+	/**
+	Get motion source for hand.
+	@return Motion source to which hand is mapped.
+	*/
+	UFUNCTION(BlueprintCallable, Category = "MotionController|MagicLeap")
+		static FName GetMotionSourceForHand(EControllerHand Hand);
+
+	/**
+	Get motion source for hand.
+	@return Hand to which motion source is mapped, ControllerHand_Count if not mapped.
+	*/
+	UFUNCTION(BlueprintCallable, Category = "MotionController|MagicLeap")
+		static EControllerHand GetHandForMotionSource(FName MotionSource);
+
+	/**
+	Set motion source for hand.
+	@return True if successful, false otherwise.
+	*/
+	UFUNCTION(BlueprintCallable, Category = "MotionController|MagicLeap")
+		static bool SetMotionSourceForHand(EControllerHand Hand, FName MotionSource);
 
+	/**
+	Type of ML device bound to the specified motion source
+	@param MotionSource Motion source to query.
+	@return Type of ML device which maps to given Unreal controller hand.
+	*/
+	UFUNCTION(BlueprintCallable, BlueprintPure, Category = "MotionController|MagicLeap")
+		static EMLControllerType GetControllerType(FName MotionSource);
+
+	/////////////////////////////////////////////////////////////////////////////////////////////
+	// DEPRECATED FUNCTIONS
+	/////////////////////////////////////////////////////////////////////////////////////////////
 	/**
 	  Gets the maximum number of Magic Leap controllers supported at a time.
 	  @return the maximum number of Magic Leap controllers supported at a time.
 	*/
-	UFUNCTION(BlueprintCallable, Category = "MotionController|MagicLeap")
+	UFUNCTION(BlueprintCallable, Category = "MotionController|MagicLeap", meta = (DeprecatedFunction))
 	static int32 MaxSupportedMagicLeapControllers();
 
 	/**
@@ -33,23 +113,23 @@ public:
 	  @param Hand Output parameter which is the hand the given index maps to. Valid only if the function returns true.
 	  @return true of the controller index maps to a valid hand, false otherwise
 	*/
-	UFUNCTION(BlueprintCallable, Category = "MotionController|MagicLeap")
+	UFUNCTION(BlueprintCallable, Category = "MotionController|MagicLeap", meta = (DeprecatedFunction, DeprecationMessage = "Use GetHandForMotionSource instead"))
 	static bool GetControllerMapping(int32 ControllerIndex, EControllerHand& Hand);
 
 	/**
-	  Inverts the controller mapping i.e. keys mapped to left hand controller will now be treated as right hand and vice-versa.
-	  @see GetControllerMapping()
+	Type of ML device being tracking the given hand.
+	@param Hand Controller hand to query.
+	@return Type of ML device which maps to given Unreal controller hand.
 	*/
-	UFUNCTION(BlueprintCallable, Category = "MotionController|MagicLeap")
-	static void InvertControllerMapping();
+	UFUNCTION(BlueprintCallable, BlueprintPure, Category = "MotionController|MagicLeap", meta = (DeprecatedFunction, DeprecationMessage = "Use GetMotionSourceForHand instead"))
+		static EMLControllerType GetMLControllerType(EControllerHand Hand);
 
 	/**
-	  Type of ML device being tracking the given hand.
-	  @param Hand Controller hand to query.
-	  @return Type of ML device which maps to given Unreal controller hand.
+	  Inverts the controller mapping i.e. keys mapped to left hand controller will now be treated as right hand and vice-versa.
+	  @see GetControllerMapping()
 	*/
-	UFUNCTION(BlueprintCallable, BlueprintPure, Category = "MotionController|MagicLeap")
-	static EMLControllerType GetMLControllerType(EControllerHand Hand);
+	UFUNCTION(BlueprintCallable, Category = "MotionController|MagicLeap", meta = (DeprecatedFunction, DeprecationMessage = "Use SetMotionSourceForHand instead"))
+	static void InvertControllerMapping();
 
 	/**
 	  Light up the LED on the Magic Leap Controller in the given pattern for the specified duration.
@@ -59,7 +139,7 @@ public:
 	  @param DurationInSec Duration (in seconds) to play the LED pattern.
 	  @return True if the command to play the LED pattern was successfully sent to the controller, false otherwise.
 	*/
-	UFUNCTION(BlueprintCallable, Category = "MotionController|MagicLeap")
+	UFUNCTION(BlueprintCallable, Category = "MotionController|MagicLeap", meta = (DeprecatedFunction, DeprecationMessage = "Use PlayLEDPattern instead"))
 	static bool PlayControllerLED(EControllerHand Hand, EMLControllerLEDPattern LEDPattern, EMLControllerLEDColor LEDColor, float DurationInSec);
 
 	/**
@@ -72,7 +152,7 @@ public:
 	  @param DurationInSec Duration (in seconds) to play the LED pattern.
 	  @return True if the command to play the LED effect was successfully sent to the controller, false otherwise.
 	*/
-	UFUNCTION(BlueprintCallable, Category = "MotionController|MagicLeap")
+	UFUNCTION(BlueprintCallable, Category = "MotionController|MagicLeap", meta = (DeprecatedFunction, DeprecationMessage = "Use PlayLEDEffect instead"))
 	static bool PlayControllerLEDEffect(EControllerHand Hand, EMLControllerLEDEffect LEDEffect, EMLControllerLEDSpeed LEDSpeed, EMLControllerLEDPattern LEDPattern, EMLControllerLEDColor LEDColor, float DurationInSec);
 
 	/**
@@ -82,21 +162,6 @@ public:
 	  @param Intensity Intensity to play on the controller.
 	  @return True if the command to play the haptic feedback was successfully sent to the controller, false otherwise.
 	*/
-	UFUNCTION(BlueprintCallable, Category = "MotionController|MagicLeap")
+	UFUNCTION(BlueprintCallable, Category = "MotionController|MagicLeap", meta = (DeprecatedFunction, DeprecationMessage = "Use PlayHapticPattern instead"))
 	static bool PlayControllerHapticFeedback(EControllerHand Hand, EMLControllerHapticPattern HapticPattern, EMLControllerHapticIntensity Intensity);
-
-	/**
-	Set controller tracking mode.
-	@param TrackingMode Controller tracking mode.
-	@return True if the command to set the tracking mode was successfully sent to the controller, false otherwise.
-	*/
-	UFUNCTION(BlueprintCallable, Category = "MotionController|MagicLeap")
-		static bool SetControllerTrackingMode(EMLControllerTrackingMode TrackingMode);
-
-	/**
-	Get controller tracking mode.
-	@return Controller tracking mode.
-	*/
-	UFUNCTION(BlueprintCallable, Category = "MotionController|MagicLeap")
-		static EMLControllerTrackingMode GetControllerTrackingMode();
 };
diff --git a/Engine/Plugins/Lumin/MagicLeap/Source/MagicLeapController/Public/MagicLeapControllerKeys.h b/Engine/Plugins/Lumin/MagicLeap/Source/MagicLeapController/Public/MagicLeapControllerKeys.h
--- a/Engine/Plugins/Lumin/MagicLeap/Source/MagicLeapController/Public/MagicLeapControllerKeys.h
+++ b/Engine/Plugins/Lumin/MagicLeap/Source/MagicLeapController/Public/MagicLeapControllerKeys.h
@@ -3,6 +3,17 @@
 #pragma once
 
 #include "UObject/NameTypes.h"
+#include "UObject/NameTypes.h"
+#include "InputCoreTypes.h"
+
+/** Magic Leap motion sources */
+struct FMagicLeapMotionSourceNames
+{
+	static const FName Control0;
+	static const FName Control1;
+	static const FName MobileApp;
+	static const FName Unknown;
+};
 
 /** List of input key names for all keys specific to Magic Leap Controller and Mobile Companion App. */
 struct FMagicLeapControllerKeyNames
@@ -18,6 +29,19 @@ struct FMagicLeapControllerKeyNames
 	static const FName Right_HomeButton_Name;
 };
 
+struct FMagicLeapKeys
+{
+	static const FKey MotionController_Left_Thumbstick_Z;
+	static const FKey Left_MoveButton;
+	static const FKey Left_AppButton;
+	static const FKey Left_HomeButton;
+
+	static const FKey MotionController_Right_Thumbstick_Z;
+	static const FKey Right_MoveButton;
+	static const FKey Right_AppButton;
+	static const FKey Right_HomeButton;
+};
+
 /** Defines the Magic Leap controller type. */
 UENUM(BlueprintType)
 enum class EMLControllerType : uint8
diff --git a/Engine/Plugins/Lumin/MagicLeap/Source/MagicLeapController/Public/TouchpadGesturesComponent.h b/Engine/Plugins/Lumin/MagicLeap/Source/MagicLeapController/Public/TouchpadGesturesComponent.h
--- a/Engine/Plugins/Lumin/MagicLeap/Source/MagicLeapController/Public/TouchpadGesturesComponent.h
+++ b/Engine/Plugins/Lumin/MagicLeap/Source/MagicLeapController/Public/TouchpadGesturesComponent.h
@@ -45,10 +45,14 @@ struct FMagicLeapTouchpadGesture
 	GENERATED_BODY()
 
 public:
-	/** Controller on which the gesture was performed. */
-	UPROPERTY(BlueprintReadOnly, Category = "TouchpadGesture|MagicLeap")
+	/** Hand on which the gesture was performed. */
+	UPROPERTY(BlueprintReadOnly, Category = "TouchpadGesture|MagicLeap", meta = (DeprecatedProperty, DeprecationMessage = "Hand is deprecated. Please use MotionSource instead."))
 	EControllerHand Hand;
 
+	/** Motion source on which the gesture was performed. */
+	UPROPERTY(BlueprintReadOnly, Category = "TouchpadGesture|MagicLeap")
+	FName MotionSource;
+
 	/** Type of gesture. */
 	UPROPERTY(BlueprintReadOnly, Category = "TouchpadGesture|MagicLeap")
 	EMagicLeapTouchpadGestureType Type;
diff --git a/Engine/Plugins/Lumin/MagicLeap/Source/MagicLeapEyeTracker/Private/MagicLeapEyeTrackerModule.cpp b/Engine/Plugins/Lumin/MagicLeap/Source/MagicLeapEyeTracker/Private/MagicLeapEyeTrackerModule.cpp
--- a/Engine/Plugins/Lumin/MagicLeap/Source/MagicLeapEyeTracker/Private/MagicLeapEyeTrackerModule.cpp
+++ b/Engine/Plugins/Lumin/MagicLeap/Source/MagicLeapEyeTracker/Private/MagicLeapEyeTrackerModule.cpp
@@ -18,6 +18,11 @@ FMagicLeapEyeTracker::FMagicLeapEyeTracker()
 }
 
 FMagicLeapEyeTracker::~FMagicLeapEyeTracker()
+{
+	Destroy();
+}
+
+void FMagicLeapEyeTracker::Destroy()
 {
 	if (VREyeTracker)
 	{
@@ -26,7 +31,6 @@ FMagicLeapEyeTracker::~FMagicLeapEyeTracker()
 	}
 }
 
-
 bool FMagicLeapEyeTracker::GetEyeTrackerGazeData(FEyeTrackerGazeData& OutGazeData) const
 {
 	if (VREyeTracker)
@@ -173,14 +177,18 @@ bool FMagicLeapEyeTracker::GetFixationComfort(FMagicLeapFixationComfort& Fixatio
 }
 
 /************************************************************************/
-/* FMagicLeapEyeTrackerModule                                                     */
+/* FMagicLeapEyeTrackerModule                                           */
 /************************************************************************/
+FMagicLeapEyeTrackerModule::FMagicLeapEyeTrackerModule()
+	: IMagicLeapModule("MagicLeapEyeTracker")
+{
+
+}
+
 void FMagicLeapEyeTrackerModule::StartupModule()
 {
 	IEyeTrackerModule::StartupModule();
-
 	MagicLeapEyeTracker = TSharedPtr<FMagicLeapEyeTracker, ESPMode::ThreadSafe>(new FMagicLeapEyeTracker());
-
 	OnDrawDebugHandle = AHUD::OnShowDebugInfo.AddRaw(this, &FMagicLeapEyeTrackerModule::OnDrawDebug);
 }
 
@@ -190,6 +198,14 @@ void FMagicLeapEyeTrackerModule::ShutdownModule()
 	FCoreUObjectDelegates::PreLoadMap.Remove(OnPreLoadMapHandle);
 }
 
+void FMagicLeapEyeTrackerModule::Disable()
+{
+	if (MagicLeapEyeTracker.IsValid())
+	{
+		MagicLeapEyeTracker->Destroy();
+	}
+}
+
 TSharedPtr<class IEyeTracker, ESPMode::ThreadSafe> FMagicLeapEyeTrackerModule::CreateEyeTracker()
 {
 	return MagicLeapEyeTracker;
@@ -199,9 +215,9 @@ void FMagicLeapEyeTrackerModule::OnDrawDebug(AHUD* HUD, UCanvas* Canvas, const F
 {
 	if (CVarEnableEyetrackingDebug.GetValueOnGameThread())
 	{
-		if (MagicLeapEyeTracker.IsValid() && MagicLeapEyeTracker->VREyeTracker)
+		if (MagicLeapEyeTracker.IsValid() && MagicLeapEyeTracker->GetVREyeTracker())
 		{
-			MagicLeapEyeTracker->VREyeTracker->DrawDebug(HUD, Canvas, DisplayInfo, YL, YPos);
+			MagicLeapEyeTracker->GetVREyeTracker()->DrawDebug(HUD, Canvas, DisplayInfo, YL, YPos);
 		}
 	}
 }
diff --git a/Engine/Plugins/Lumin/MagicLeap/Source/MagicLeapEyeTracker/Public/MagicLeapEyeTrackerModule.h b/Engine/Plugins/Lumin/MagicLeap/Source/MagicLeapEyeTracker/Public/MagicLeapEyeTrackerModule.h
--- a/Engine/Plugins/Lumin/MagicLeap/Source/MagicLeapEyeTracker/Public/MagicLeapEyeTrackerModule.h
+++ b/Engine/Plugins/Lumin/MagicLeap/Source/MagicLeapEyeTracker/Public/MagicLeapEyeTrackerModule.h
@@ -9,6 +9,7 @@
 #include "GameFramework/HUD.h"
 #include "Kismet/BlueprintFunctionLibrary.h"
 #include "MagicLeapEyeTrackerTypes.h"
+#include "IMagicLeapModule.h"
 #include "MagicLeapEyeTrackerModule.generated.h"
 
 class FMagicLeapVREyeTracker;
@@ -70,7 +71,7 @@ public:
 	{
 		return FModuleManager::LoadModuleChecked<IMagicLeapEyeTrackerModule>("MagicLeapEyeTracker");
 	}
-
+	
 	/**
 	* Checks to see if this module is loaded and ready.  It is only valid to call Get() if
 	* IsAvailable() returns true.
@@ -99,6 +100,7 @@ public:
 	/************************************************************************/
 	/* IEyeTracker                                                          */
 	/************************************************************************/
+	void Destroy();
 	virtual void SetEyeTrackedPlayer(APlayerController* PlayerController) override;
 	virtual bool GetEyeTrackerGazeData(FEyeTrackerGazeData& OutGazeData) const override;
 	virtual bool GetEyeTrackerStereoGazeData(FEyeTrackerStereoGazeData& OutGazeData) const override;
@@ -111,20 +113,25 @@ public:
 
 	EMagicLeapEyeTrackingCalibrationStatus GetCalibrationStatus() const;
 
-private:
-	friend class FMagicLeapEyeTrackerModule;
+	inline FMagicLeapVREyeTracker* GetVREyeTracker() const
+	{
+		return VREyeTracker;
+	}
 
+private:
 	FMagicLeapVREyeTracker* VREyeTracker;
 };
 
-class FMagicLeapEyeTrackerModule : public IMagicLeapEyeTrackerModule
+class FMagicLeapEyeTrackerModule : public IMagicLeapEyeTrackerModule, public IMagicLeapModule
 {
 	/************************************************************************/
 	/* IInputDeviceModule                                                   */
 	/************************************************************************/
 public:
+	FMagicLeapEyeTrackerModule();
 	virtual void StartupModule() override;
 	virtual void ShutdownModule() override;
+	virtual void Disable() override;
 	virtual TSharedPtr< class IEyeTracker, ESPMode::ThreadSafe > CreateEyeTracker() override;
 
 	/************************************************************************/
diff --git a/Engine/Plugins/Lumin/MagicLeap/Source/MagicLeapHandTracking/Private/MagicLeapHandTracking.cpp b/Engine/Plugins/Lumin/MagicLeap/Source/MagicLeapHandTracking/Private/MagicLeapHandTracking.cpp
--- a/Engine/Plugins/Lumin/MagicLeap/Source/MagicLeapHandTracking/Private/MagicLeapHandTracking.cpp
+++ b/Engine/Plugins/Lumin/MagicLeap/Source/MagicLeapHandTracking/Private/MagicLeapHandTracking.cpp
@@ -333,6 +333,9 @@ FMagicLeapHandTracking::FMagicLeapHandTracking(const TSharedRef<FGenericApplicat
 	, bIsHandTrackingStateValid(false)
 {
 #if WITH_MLSDK
+	// Zero-out structs
+	FMemory::Memset(&HandTrackingDatas, 0, sizeof(HandTrackingDatas));
+
 	// Initialize gesture data to default values.
 	HandTrackingDatas[0].left_hand_state.keypose = MLHandTrackingKeyPose_NoHand;
 	HandTrackingDatas[0].right_hand_state.keypose = MLHandTrackingKeyPose_NoHand;
@@ -351,11 +354,7 @@ FMagicLeapHandTracking::FMagicLeapHandTracking(const TSharedRef<FGenericApplicat
 
 	// We're implicitly requiring that the MagicLeapPlugin has been loaded and
 	// initialized at this point.
-	auto HMD = IMagicLeapPlugin::Get().GetHMD().Pin();
-	if (HMD.IsValid())
-	{
-		HMD->RegisterMagicLeapInputDevice(this);
-	}
+	IMagicLeapPlugin::Get().RegisterMagicLeapInputDevice(this);
 }
 
 FMagicLeapHandTracking::~FMagicLeapHandTracking()
@@ -364,11 +363,7 @@ FMagicLeapHandTracking::~FMagicLeapHandTracking()
 	// but it isn't an assumption that we should make.
 	if (IMagicLeapPlugin::IsAvailable())
 	{
-		auto HMD = IMagicLeapPlugin::Get().GetHMD().Pin();
-		if (HMD.IsValid())
-		{
-			HMD->UnregisterMagicLeapInputDevice(this);
-		}
+		IMagicLeapPlugin::Get().UnregisterMagicLeapInputDevice(this);
 	}
 
 	Disable();
@@ -742,6 +737,8 @@ void FMagicLeapHandTracking::EnumerateSources(TArray<FMotionControllerSource>& S
 void FMagicLeapHandTracking::Tick(float DeltaTime)
 {
 	UpdateTrackerData();
+	UpdateCurrentHandTrackerTransforms();
+	UpdateLiveLink();
 }
 
 void FMagicLeapHandTracking::SendControllerEvents()
@@ -831,6 +828,13 @@ void FMagicLeapHandTracking::Disable()
 #endif
 }
 
+void FMagicLeapHandTracking::OnBeginRendering_GameThread_Update()
+{
+#if WITH_MLSDK
+	UpdateCurrentHandTrackerTransforms();
+#endif
+}
+
 const FMagicLeapHandTracking::FHandState& FMagicLeapHandTracking::GetLeftHandState() const
 {
 	return LeftHand;
@@ -1088,18 +1092,21 @@ namespace MagicLeapHandTracking
 
 void FMagicLeapHandTracking::UpdateTrackerData()
 {
+	// This updates the tracker data and live link clients while "consuming" a data
+	// slot for the frame.
 #if WITH_MLSDK
+	check(IsInGameThread());
+
 	if (MLHandleIsValid(HandTracker))
 	{
 		CurrentHandTrackingDataIndex = 1 - CurrentHandTrackingDataIndex;
-		bIsHandTrackingStateValid =		MLHandTrackingGetData(HandTracker, &HandTrackingDatas[CurrentHandTrackingDataIndex]) == MLResult_Ok;
+		bIsHandTrackingStateValid = MLHandTrackingGetData(HandTracker, &HandTrackingDatas[CurrentHandTrackingDataIndex]) == MLResult_Ok;
 		//MagicLeapHandTracking::LogMLHandTrackingData(GetCurrentHandTrackingData());
 	}
 	else
 	{
 		bIsHandTrackingStateValid = false;
 	}
-
 	if (bIsHandTrackingStateValid && IMagicLeapPlugin::Get().IsMagicLeapHMDValid())
 	{
 		const FAppFramework& AppFramework = static_cast<FMagicLeapHMD*>(GEngine->XRSystem->GetHMDDevice())->GetAppFrameworkConst();
@@ -1108,6 +1115,26 @@ void FMagicLeapHandTracking::UpdateTrackerData()
 		const MLHandTrackingData& HandTrackingData = GetCurrentHandTrackingData();
 
 		LeftHand.Gesture = TranslateGestureEnum(HandTrackingData.left_hand_state.keypose);
+		RightHand.Gesture = TranslateGestureEnum(HandTrackingData.right_hand_state.keypose);
+
+		LeftHand.GestureConfidence = HandTrackingData.left_hand_state.keypose_confidence[HandTrackingData.left_hand_state.keypose];
+		RightHand.GestureConfidence = HandTrackingData.right_hand_state.keypose_confidence[HandTrackingData.right_hand_state.keypose];
+	}
+#endif //WITH_MLSDK
+}
+
+void FMagicLeapHandTracking::UpdateCurrentHandTrackerTransforms()
+{
+#if WITH_MLSDK
+	check(IsInGameThread());
+
+	if (MLHandleIsValid(HandTracker) && bIsHandTrackingStateValid && IMagicLeapPlugin::Get().IsMagicLeapHMDValid())
+	{
+		const FAppFramework& AppFramework = static_cast<FMagicLeapHMD*>(GEngine->XRSystem->GetHMDDevice())->GetAppFrameworkConst();
+		check(AppFramework.IsInitialized());
+
+		const MLHandTrackingData& HandTrackingData = GetCurrentHandTrackingData();
+
 		if (LeftHand.Gesture != EHandTrackingGesture::NoHand)
 		{
 			LeftHand.HandCenterNormalized = MagicLeap::ToFVector(HandTrackingData.left_hand_state.hand_center_normalized, 1.0f);
@@ -1119,8 +1146,6 @@ void FMagicLeapHandTracking::UpdateTrackerData()
 
 			MagicLeapHandTracking::FetchHandTransforms(AppFramework, HandTrackingStaticData.left, LeftHand, TEXT("left hand"));
 		}
-
-		RightHand.Gesture = TranslateGestureEnum(HandTrackingData.right_hand_state.keypose);
 		if (RightHand.Gesture != EHandTrackingGesture::NoHand)
 		{
 			RightHand.HandCenterNormalized = MagicLeap::ToFVector(HandTrackingData.right_hand_state.hand_center_normalized, 1.0f);
@@ -1132,13 +1157,8 @@ void FMagicLeapHandTracking::UpdateTrackerData()
 
 			MagicLeapHandTracking::FetchHandTransforms(AppFramework, HandTrackingStaticData.right, RightHand, TEXT("right hand"));
 		}
-
-		LeftHand.GestureConfidence = HandTrackingData.left_hand_state.keypose_confidence[HandTrackingData.left_hand_state.keypose];
-		RightHand.GestureConfidence = HandTrackingData.right_hand_state.keypose_confidence[HandTrackingData.right_hand_state.keypose];
-
-		UpdateLiveLink();
 	}
-#endif //WITH_MLSDK
+#endif
 }
 
 void FMagicLeapHandTracking::AddKeys()
diff --git a/Engine/Plugins/Lumin/MagicLeap/Source/MagicLeapHandTracking/Private/MagicLeapHandTracking.h b/Engine/Plugins/Lumin/MagicLeap/Source/MagicLeapHandTracking/Private/MagicLeapHandTracking.h
--- a/Engine/Plugins/Lumin/MagicLeap/Source/MagicLeapHandTracking/Private/MagicLeapHandTracking.h
+++ b/Engine/Plugins/Lumin/MagicLeap/Source/MagicLeapHandTracking/Private/MagicLeapHandTracking.h
@@ -199,6 +199,7 @@ public:
 	virtual void Enable() override;
 	virtual bool SupportsExplicitEnable() const override;
 	virtual void Disable() override;
+	virtual void OnBeginRendering_GameThread_Update() override;
 
 	const FHandState& GetLeftHandState() const;
 	const FHandState& GetRightHandState() const;
@@ -217,6 +218,7 @@ private:
 	const FTransform* FindTransformBySource(FName SourceName) const;
 	const FHandState* FindHandBySource(FName SourceName) const;
 	void UpdateTrackerData();
+	void UpdateCurrentHandTrackerTransforms();
 	void AddKeys();
 	void ConditionallyEnable();
 	void OnAppPause() override;
diff --git a/Engine/Plugins/Lumin/MagicLeap/Source/MagicLeapHandTracking/Private/MagicLeapHandTrackingLiveLink.cpp b/Engine/Plugins/Lumin/MagicLeap/Source/MagicLeapHandTracking/Private/MagicLeapHandTrackingLiveLink.cpp
--- a/Engine/Plugins/Lumin/MagicLeap/Source/MagicLeapHandTracking/Private/MagicLeapHandTrackingLiveLink.cpp
+++ b/Engine/Plugins/Lumin/MagicLeap/Source/MagicLeapHandTracking/Private/MagicLeapHandTrackingLiveLink.cpp
@@ -7,6 +7,7 @@
 #include "Engine/Engine.h"
 #include "LiveLinkSourceFactory.h"
 #include "ILiveLinkClient.h"
+#include "IMagicLeapPlugin.h"
 
 #define LOCTEXT_NAMESPACE "MagicLeapHandTracking"
 
@@ -146,7 +147,7 @@ void FMagicLeapHandTracking::UpdateLiveLink()
 {
 	check(IsInGameThread());
 
-	if (LiveLinkClient)
+	if (LiveLinkClient && bIsHandTrackingStateValid && IMagicLeapPlugin::Get().IsMagicLeapHMDValid())
 	{
 		if (bNewLiveLinkClient)
 		{
diff --git a/Engine/Plugins/Lumin/MagicLeap/Source/MagicLeapHelperVulkan/Private/MagicLeapHelperVulkan.cpp b/Engine/Plugins/Lumin/MagicLeap/Source/MagicLeapHelperVulkan/Private/MagicLeapHelperVulkan.cpp
--- a/Engine/Plugins/Lumin/MagicLeap/Source/MagicLeapHelperVulkan/Private/MagicLeapHelperVulkan.cpp
+++ b/Engine/Plugins/Lumin/MagicLeap/Source/MagicLeapHelperVulkan/Private/MagicLeapHelperVulkan.cpp
@@ -3,6 +3,7 @@
 #include "MagicLeapHelperVulkan.h"
 #include "IMagicLeapHelperVulkanPlugin.h"
 #include "Engine/Engine.h"
+#include "XRThreadUtils.h"
 
 #if !PLATFORM_MAC
 #include "VulkanRHIPrivate.h"
@@ -62,10 +63,10 @@ void FMagicLeapHelperVulkan::BlitImage(uint64 SrcName, int32 SrcLevel, int32 Src
 	Region.srcSubresource.aspectMask = VK_IMAGE_ASPECT_COLOR_BIT;
 	Region.srcSubresource.layerCount = 1;
 	Region.dstOffsets[0].x = DstX;
-	Region.dstOffsets[0].y = DstY;
+	Region.dstOffsets[0].y = DstY + DstHeight;
 	Region.dstOffsets[0].z = DstZ;
 	Region.dstOffsets[1].x = DstX + DstWidth;
-	Region.dstOffsets[1].y = DstY + DstHeight;
+	Region.dstOffsets[1].y = DstY;
 	Region.dstOffsets[1].z = DstZ + DstDepth;
 	Region.dstSubresource.aspectMask = VK_IMAGE_ASPECT_COLOR_BIT;
 	Region.dstSubresource.baseArrayLayer = DstLevel;
@@ -218,36 +219,40 @@ bool FMagicLeapHelperVulkan::GetMediaTexture(FTextureRHIRef& ResultTexture, FSam
 		return false;
 	}
 
-	VkImageMemoryBarrier ImageBarrier;
-	FMemory::Memzero(&ImageBarrier, sizeof(VkImageMemoryBarrier));
-	ImageBarrier.sType = VK_STRUCTURE_TYPE_IMAGE_MEMORY_BARRIER;
-	ImageBarrier.pNext = nullptr;
-	ImageBarrier.srcAccessMask = 0;
-	ImageBarrier.dstAccessMask = VK_ACCESS_SHADER_READ_BIT;
-	ImageBarrier.oldLayout = VK_IMAGE_LAYOUT_PREINITIALIZED;
-	ImageBarrier.newLayout = VK_IMAGE_LAYOUT_SHADER_READ_ONLY_OPTIMAL;
-	ImageBarrier.srcQueueFamilyIndex = VK_QUEUE_FAMILY_IGNORED;
-	ImageBarrier.dstQueueFamilyIndex = VK_QUEUE_FAMILY_IGNORED;
-	ImageBarrier.image = MediaSurface.imported_image;
-	ImageBarrier.subresourceRange.aspectMask = VK_IMAGE_ASPECT_COLOR_BIT;
-	ImageBarrier.subresourceRange.baseMipLevel = 0;
-	ImageBarrier.subresourceRange.levelCount = 1;
-	ImageBarrier.subresourceRange.baseArrayLayer = 0;
-	ImageBarrier.subresourceRange.layerCount = 1;
-
-	FVulkanCommandListContext& ImmediateContext = Device->GetImmediateContext();
-	FVulkanCmdBuffer* CmdBuffer = ImmediateContext.GetCommandBufferManager()->GetUploadCmdBuffer();
-	VulkanRHI::vkCmdPipelineBarrier(
-		CmdBuffer->GetHandle(),
-		VK_PIPELINE_STAGE_TOP_OF_PIPE_BIT,
-		VK_PIPELINE_STAGE_FRAGMENT_SHADER_BIT,
-		0,
-		0,
-		nullptr,
-		0,
-		nullptr,
-		1,
-		&ImageBarrier);
+	VkImage ImportedImage = MediaSurface.imported_image;
+	ExecuteOnRHIThread_DoNotWait([Device, ImportedImage]()
+	{
+		VkImageMemoryBarrier ImageBarrier;
+		FMemory::Memzero(&ImageBarrier, sizeof(VkImageMemoryBarrier));
+		ImageBarrier.sType = VK_STRUCTURE_TYPE_IMAGE_MEMORY_BARRIER;
+		ImageBarrier.pNext = nullptr;
+		ImageBarrier.srcAccessMask = 0;
+		ImageBarrier.dstAccessMask = VK_ACCESS_SHADER_READ_BIT;
+		ImageBarrier.oldLayout = VK_IMAGE_LAYOUT_PREINITIALIZED;
+		ImageBarrier.newLayout = VK_IMAGE_LAYOUT_SHADER_READ_ONLY_OPTIMAL;
+		ImageBarrier.srcQueueFamilyIndex = VK_QUEUE_FAMILY_IGNORED;
+		ImageBarrier.dstQueueFamilyIndex = VK_QUEUE_FAMILY_IGNORED;
+		ImageBarrier.image = ImportedImage;
+		ImageBarrier.subresourceRange.aspectMask = VK_IMAGE_ASPECT_COLOR_BIT;
+		ImageBarrier.subresourceRange.baseMipLevel = 0;
+		ImageBarrier.subresourceRange.levelCount = 1;
+		ImageBarrier.subresourceRange.baseArrayLayer = 0;
+		ImageBarrier.subresourceRange.layerCount = 1;
+
+		FVulkanCommandListContext& ImmediateContext = Device->GetImmediateContext();
+		FVulkanCmdBuffer* CmdBuffer = ImmediateContext.GetCommandBufferManager()->GetUploadCmdBuffer();
+		VulkanRHI::vkCmdPipelineBarrier(
+			CmdBuffer->GetHandle(),
+			VK_PIPELINE_STAGE_TOP_OF_PIPE_BIT,
+			VK_PIPELINE_STAGE_FRAGMENT_SHADER_BIT,
+			0,
+			0,
+			nullptr,
+			0,
+			nullptr,
+			1,
+			&ImageBarrier);
+	});
 
 	FSamplerYcbcrConversionInitializer ConversionInitializer;
 	FMemory::Memzero(&ConversionInitializer, sizeof(FSamplerYcbcrConversionInitializer));
diff --git a/Engine/Plugins/Lumin/MagicLeapCamera/MagicLeapCamera.uplugin b/Engine/Plugins/Lumin/MagicLeapCamera/MagicLeapCamera.uplugin
new file mode 100644
--- /dev/null
+++ b/Engine/Plugins/Lumin/MagicLeapCamera/MagicLeapCamera.uplugin
@@ -0,0 +1,38 @@
+{
+	"FileVersion": 3,
+	"Version": 1,
+	"VersionName": "1.0",
+	"FriendlyName": "Magic Leap Camera",
+	"CreatedBy": "Magic Leap, Inc.",
+	"CreatedByURL": "http://magicleap.com",
+	"Description": "Support for the Magic Leap Camera API.",
+	"Category": "Lumin",
+	"EnabledByDefault": false,
+	"DocsURL": "",
+	"MarketplaceURL": "",
+	"SupportURL": "",
+	"CanContainContent": false,
+	"IsBetaVersion": false,
+	"Installed": false,
+	"SupportedTargetPlatforms": [
+		"Lumin"
+	],
+	"Modules": [
+		{
+			"Name": "MagicLeapCamera",
+			"Type": "Runtime",
+			"LoadingPhase": "PostConfigInit",
+			"WhitelistPlatforms": [
+				"Win64",
+				"Mac",
+				"Lumin"
+			]
+		}
+	],
+	"Plugins": [
+		{
+			"Name": "MagicLeap",
+			"Enabled": true
+		}
+	]
+}
diff --git a/Engine/Plugins/Lumin/MagicLeapCamera/Source/MagicLeapCamera.Build.cs b/Engine/Plugins/Lumin/MagicLeapCamera/Source/MagicLeapCamera.Build.cs
new file mode 100644
--- /dev/null
+++ b/Engine/Plugins/Lumin/MagicLeapCamera/Source/MagicLeapCamera.Build.cs
@@ -0,0 +1,43 @@
+// Copyright 1998-2018 Epic Games, Inc. All Rights Reserved.
+
+using System.IO;
+
+namespace UnrealBuildTool.Rules
+{
+	public class MagicLeapCamera : ModuleRules
+	{
+		public MagicLeapCamera( ReadOnlyTargetRules Target ) : base(Target)
+		{
+			PublicIncludePaths.Add(Path.Combine(ModuleDirectory, "Public"));            
+
+            PublicDependencyModuleNames.AddRange(
+				new string[]
+				{
+					"Core",
+					"CoreUObject",
+					"Engine"
+				}
+			);
+
+			PrivateDependencyModuleNames.AddRange(
+				new string[]
+				{
+					"LuminRuntimeSettings",
+					"MLSDK",
+					"MagicLeap",
+					"HeadMountedDisplay",
+				}
+			);
+
+			// This is not ideal but needs to be done in order to expose the private MagicLeapHMD headers to this module.
+			PrivateIncludePaths.Add(Path.Combine(
+				new string[] 
+				{
+					ModuleDirectory,
+					"..", "..",
+					"MagicLeap", "Source", "MagicLeap", "Private"
+                })
+			);
+        }
+	}
+}
diff --git a/Engine/Plugins/Lumin/MagicLeapCamera/Source/Private/MagicLeapCameraComponent.cpp b/Engine/Plugins/Lumin/MagicLeapCamera/Source/Private/MagicLeapCameraComponent.cpp
new file mode 100644
--- /dev/null
+++ b/Engine/Plugins/Lumin/MagicLeapCamera/Source/Private/MagicLeapCameraComponent.cpp
@@ -0,0 +1,42 @@
+// Copyright 1998-2018 Epic Games, Inc. All Rights Reserved.
+
+#include "MagicLeapCameraComponent.h"
+#include "MagicLeapCameraPlugin.h"
+
+void UMagicLeapCameraComponent::BeginPlay()
+{
+	Super::BeginPlay();
+	GET_MAGIC_LEAP_CAMERA_PLUGIN()->IncUserCount();
+	GET_MAGIC_LEAP_CAMERA_PLUGIN()->SetLogDelegate(OnLogMessage);
+}
+
+void UMagicLeapCameraComponent::EndPlay(const EEndPlayReason::Type EndPlayReason)
+{
+	GET_MAGIC_LEAP_CAMERA_PLUGIN()->DecUserCount();
+	Super::EndPlay(EndPlayReason);
+}
+
+bool UMagicLeapCameraComponent::CaptureImageToFileAsync()
+{
+	return GET_MAGIC_LEAP_CAMERA_PLUGIN()->CaptureImageToFileAsync(OnCaptureImgToFile);
+}
+
+bool UMagicLeapCameraComponent::CaptureImageToTextureAsync()
+{
+	return GET_MAGIC_LEAP_CAMERA_PLUGIN()->CaptureImageToTextureAsync(OnCaptureImgToTexture);
+}
+
+bool UMagicLeapCameraComponent::StartRecordingAsync()
+{
+	return GET_MAGIC_LEAP_CAMERA_PLUGIN()->StartRecordingAsync(OnStartRecording);
+}
+
+bool UMagicLeapCameraComponent::StopRecordingAsync()
+{
+	return GET_MAGIC_LEAP_CAMERA_PLUGIN()->StopRecordingAsync(OnStopRecording);
+}
+
+bool UMagicLeapCameraComponent::IsCapturing() const
+{
+	return GET_MAGIC_LEAP_CAMERA_PLUGIN()->IsCapturing();
+}
diff --git a/Engine/Plugins/Lumin/MagicLeapCamera/Source/Private/MagicLeapCameraFunctionLibrary.cpp b/Engine/Plugins/Lumin/MagicLeapCamera/Source/Private/MagicLeapCameraFunctionLibrary.cpp
new file mode 100644
--- /dev/null
+++ b/Engine/Plugins/Lumin/MagicLeapCamera/Source/Private/MagicLeapCameraFunctionLibrary.cpp
@@ -0,0 +1,54 @@
+// Copyright 1998-2018 Epic Games, Inc. All Rights Reserved.
+
+#include "MagicLeapCameraFunctionLibrary.h"
+#include "MagicLeapCameraPlugin.h"
+
+bool UMagicLeapCameraFunctionLibrary::CameraConnect(const FCameraConnect& ResultDelegate)
+{
+	return GET_MAGIC_LEAP_CAMERA_PLUGIN()->CameraConnect(ResultDelegate);
+}
+
+bool UMagicLeapCameraFunctionLibrary::CameraDisconnect(const FCameraDisconnect& ResultDelegate)
+{
+	return GET_MAGIC_LEAP_CAMERA_PLUGIN()->CameraDisconnect(ResultDelegate);
+}
+
+bool UMagicLeapCameraFunctionLibrary::CaptureImageToFileAsync(const FCameraCaptureImgToFile& InResultDelegate)
+{
+	FCameraCaptureImgToFileMulti ResultDelegate;
+	ResultDelegate.Add(InResultDelegate);
+	return GET_MAGIC_LEAP_CAMERA_PLUGIN()->CaptureImageToFileAsync(ResultDelegate);
+}
+
+bool UMagicLeapCameraFunctionLibrary::CaptureImageToTextureAsync(const FCameraCaptureImgToTexture& InResultDelegate)
+{
+	FCameraCaptureImgToTextureMulti ResultDelegate;
+	ResultDelegate.Add(InResultDelegate);
+	return GET_MAGIC_LEAP_CAMERA_PLUGIN()->CaptureImageToTextureAsync(ResultDelegate);
+}
+
+bool UMagicLeapCameraFunctionLibrary::StartRecordingAsync(const FCameraStartRecording& InResultDelegate)
+{
+	FCameraStartRecordingMulti ResultDelegate;
+	ResultDelegate.Add(InResultDelegate);
+	return GET_MAGIC_LEAP_CAMERA_PLUGIN()->StartRecordingAsync(ResultDelegate);
+}
+
+bool UMagicLeapCameraFunctionLibrary::StopRecordingAsync(const FCameraStopRecording& InResultDelegate)
+{
+	FCameraStopRecordingMulti ResultDelegate;
+	ResultDelegate.Add(InResultDelegate);
+	return GET_MAGIC_LEAP_CAMERA_PLUGIN()->StopRecordingAsync(ResultDelegate);
+}
+
+bool UMagicLeapCameraFunctionLibrary::SetLogDelegate(const FCameraLogMessage& InLogDelegate)
+{
+	FCameraLogMessageMulti LogDelegate;
+	LogDelegate.Add(InLogDelegate);
+	return GET_MAGIC_LEAP_CAMERA_PLUGIN()->SetLogDelegate(LogDelegate);
+}
+
+bool UMagicLeapCameraFunctionLibrary::IsCapturing()
+{
+	return GET_MAGIC_LEAP_CAMERA_PLUGIN()->IsCapturing();
+}
diff --git a/Engine/Plugins/Lumin/MagicLeapCamera/Source/Private/MagicLeapCameraPlugin.cpp b/Engine/Plugins/Lumin/MagicLeapCamera/Source/Private/MagicLeapCameraPlugin.cpp
new file mode 100644
--- /dev/null
+++ b/Engine/Plugins/Lumin/MagicLeapCamera/Source/Private/MagicLeapCameraPlugin.cpp
@@ -0,0 +1,287 @@
+// Copyright 1998-2018 Epic Games, Inc. All Rights Reserved.
+
+#include "MagicLeapCameraPlugin.h"
+#include "Async/Async.h"
+
+#if WITH_MLSDK
+ML_INCLUDES_START
+#include <ml_camera.h>
+ML_INCLUDES_END
+#endif //WITH_MLSDK
+
+//PRAGMA_DISABLE_OPTIMIZATION
+
+DEFINE_LOG_CATEGORY(LogMagicLeapCamera);
+
+FMagicLeapCameraPlugin::FMagicLeapCameraPlugin()
+: UserCount(0)
+, Runnable(nullptr)
+, CurrentTaskType(FCameraTask::EType::None)
+, PrevTaskType(FCameraTask::EType::None)
+{
+}
+
+void FMagicLeapCameraPlugin::StartupModule()
+{
+	IMagicLeapCameraPlugin::StartupModule();
+	APISetup.Startup();
+#if WITH_MLSDK
+	APISetup.LoadDLL(TEXT("ml_camera"));
+#endif // WITH_MLSDK
+	Runnable = new FCameraRunnable();
+	TickDelegate = FTickerDelegate::CreateRaw(this, &FMagicLeapCameraPlugin::Tick);
+	TickDelegateHandle = FTicker::GetCoreTicker().AddTicker(TickDelegate);
+}
+
+void FMagicLeapCameraPlugin::ShutdownModule()
+{
+	FCameraRunnable* InRunnable = Runnable;
+	Runnable = nullptr;
+	AsyncTask(ENamedThreads::AnyBackgroundThreadNormalTask, [InRunnable]()
+	{
+		delete InRunnable;
+	});
+	APISetup.Shutdown();
+	FTicker::GetCoreTicker().RemoveTicker(TickDelegateHandle);
+	IModuleInterface::ShutdownModule();
+}
+
+bool FMagicLeapCameraPlugin::Tick(float DeltaTime)
+{
+	FCameraTask CompletedTask;
+	if (Runnable->TryGetCompletedTask(CompletedTask))
+	{
+		switch (CompletedTask.CaptureType)
+		{
+		case FCameraTask::EType::Connect:
+		{
+			if (OnCameraConnect.IsBound())
+			{
+				OnCameraConnect.ExecuteIfBound(CompletedTask.bSuccess);
+				// Only clear CurrentTaskType if connect was manually called
+				// (otherwise it was auto-called from within the runnable and the actual
+				// current task is still in progress).
+				CurrentTaskType = FCameraTask::EType::None;
+			}
+		}
+		break;
+
+		case FCameraTask::EType::Disconnect:
+		{
+			OnCameraDisconnect.ExecuteIfBound(CompletedTask.bSuccess);
+			CurrentTaskType = FCameraTask::EType::None;
+		}
+		break;
+
+		case FCameraTask::EType::ImageToFile:
+		{
+			OnCaptureImgToFile.Broadcast(CompletedTask.bSuccess, CompletedTask.FilePath);
+			CurrentTaskType = FCameraTask::EType::None;
+		}
+		break;
+
+		case FCameraTask::EType::ImageToTexture:
+		{
+			OnCaptureImgToTexture.Broadcast(CompletedTask.bSuccess, CompletedTask.Texture);
+			CurrentTaskType = FCameraTask::EType::None;
+		}
+		break;
+
+		case FCameraTask::EType::StartVideoToFile:
+		{
+			OnStartRecording.Broadcast(CompletedTask.bSuccess);
+			// do not reset CurrentTaskType if start recording is successful
+			// as this constitutes an ongoing capture state
+			if (!CompletedTask.bSuccess)
+			{
+				CurrentTaskType = FCameraTask::EType::None;
+			}
+		}
+		break;
+
+		case FCameraTask::EType::StopVideoToFile:
+		{
+			OnStopRecording.Broadcast(CompletedTask.bSuccess, CompletedTask.FilePath);
+			CurrentTaskType = FCameraTask::EType::None;
+		}
+		break;
+
+		case FCameraTask::EType::Log:
+		{
+			UE_LOG(LogMagicLeapCamera, Log, TEXT("%s"), *CompletedTask.Log);
+			OnLogMessage.Broadcast(FString::Printf(TEXT("<br>%s"), *CompletedTask.Log));
+		}
+		break;
+		}
+	}
+
+	return true;
+}
+
+bool FMagicLeapCameraPlugin::CameraConnect(const FCameraConnect& ResultDelegate)
+{
+	if (TryPushNewCaptureTask(FCameraTask::EType::Connect))
+	{
+		OnCameraConnect = ResultDelegate;
+		return true;
+	}
+
+	return false;
+}
+
+bool FMagicLeapCameraPlugin::CameraDisconnect(const FCameraDisconnect& ResultDelegate)
+{
+	if (UserCount <= 0 && TryPushNewCaptureTask(FCameraTask::EType::Disconnect))
+	{
+		OnCameraDisconnect = ResultDelegate;
+		return true;
+	}
+
+	return false;
+}
+
+int64 FMagicLeapCameraPlugin::GetPreviewHandle() const
+{
+	return FCameraRunnable::PreviewHandle.GetValue();
+}
+
+void FMagicLeapCameraPlugin::IncUserCount()
+{
+	++UserCount;
+}
+
+void FMagicLeapCameraPlugin::DecUserCount()
+{
+	--UserCount;
+	if (UserCount <= 0)
+	{
+		UserCount = 0;
+		TryPushNewCaptureTask(FCameraTask::EType::Disconnect);
+	}
+}
+
+bool FMagicLeapCameraPlugin::CaptureImageToFileAsync(const FCameraCaptureImgToFileMulti& ResultDelegate)
+{
+	if (TryPushNewCaptureTask(FCameraTask::EType::ImageToFile))
+	{
+		OnCaptureImgToFile = ResultDelegate;
+		return true;
+	}
+
+	return false;
+}
+
+bool FMagicLeapCameraPlugin::CaptureImageToTextureAsync(const FCameraCaptureImgToTextureMulti& ResultDelegate)
+{
+	if (TryPushNewCaptureTask(FCameraTask::EType::ImageToTexture))
+	{
+		OnCaptureImgToTexture = ResultDelegate;
+		return true;
+	}
+
+	return false;
+}
+
+bool FMagicLeapCameraPlugin::StartRecordingAsync(const FCameraStartRecordingMulti& ResultDelegate)
+{
+	if (TryPushNewCaptureTask(FCameraTask::EType::StartVideoToFile))
+	{
+		OnStartRecording = ResultDelegate;
+		return true;
+	}
+
+	return false;
+}
+
+bool FMagicLeapCameraPlugin::StopRecordingAsync(const FCameraStopRecordingMulti& ResultDelegate)
+{
+	if (TryPushNewCaptureTask(FCameraTask::EType::StopVideoToFile))
+	{
+		OnStopRecording = ResultDelegate;
+		return true;
+	}
+
+	return false;
+}
+
+bool FMagicLeapCameraPlugin::SetLogDelegate(const FCameraLogMessageMulti& LogDelegate)
+{
+	OnLogMessage = LogDelegate;
+	return true;
+}
+
+bool FMagicLeapCameraPlugin::IsCapturing() const
+{
+	return CurrentTaskType != FCameraTask::EType::None;
+}
+
+bool FMagicLeapCameraPlugin::TryPushNewCaptureTask(FCameraTask::EType InTaskType)
+{
+	bool bCanPushTask = false;
+
+	switch (InTaskType)
+	{
+	case FCameraTask::EType::None:
+	{
+		bCanPushTask = false;
+	}
+	break;
+
+	case FCameraTask::EType::Connect:
+	{
+		bCanPushTask = CurrentTaskType == FCameraTask::EType::None || CurrentTaskType == FCameraTask::EType::Disconnect;
+	}
+	break;
+
+	case FCameraTask::EType::Disconnect:
+	{
+		bCanPushTask = CurrentTaskType == FCameraTask::EType::None || CurrentTaskType == FCameraTask::EType::Connect;
+	}
+	break;
+
+	case FCameraTask::EType::ImageToFile:
+	{
+		bCanPushTask = CurrentTaskType == FCameraTask::EType::None;
+	}
+	break;
+
+	case FCameraTask::EType::ImageToTexture:
+	{
+		bCanPushTask = CurrentTaskType == FCameraTask::EType::None;
+	}
+	break;
+
+	case FCameraTask::EType::StartVideoToFile:
+	{
+		bCanPushTask = CurrentTaskType == FCameraTask::EType::None && PrevTaskType != FCameraTask::EType::StartVideoToFile;
+	}
+	break;
+
+	case FCameraTask::EType::StopVideoToFile:
+	{
+		bCanPushTask = PrevTaskType != FCameraTask::EType::StopVideoToFile &&
+			(CurrentTaskType == FCameraTask::EType::None || CurrentTaskType == FCameraTask::EType::StartVideoToFile);
+	}
+	break;
+
+	case FCameraTask::EType::Log:
+	{
+		bCanPushTask = true;
+	}
+	break;
+	}
+
+	if (bCanPushTask)
+	{
+		if (InTaskType != FCameraTask::EType::Log)
+		{
+			PrevTaskType = CurrentTaskType;
+			CurrentTaskType = InTaskType;
+		}
+
+		Runnable->PushNewCaptureTask(InTaskType);
+		return true;
+	}
+
+	return false;
+}
diff --git a/Engine/Plugins/Lumin/MagicLeapCamera/Source/Private/MagicLeapCameraPlugin.h b/Engine/Plugins/Lumin/MagicLeapCamera/Source/Private/MagicLeapCameraPlugin.h
new file mode 100644
--- /dev/null
+++ b/Engine/Plugins/Lumin/MagicLeapCamera/Source/Private/MagicLeapCameraPlugin.h
@@ -0,0 +1,54 @@
+// Copyright 1998-2018 Epic Games, Inc. All Rights Reserved.
+
+#pragma once
+
+#include "CoreMinimal.h"
+#include "IMagicLeapPlugin.h"
+#include "IMagicLeapCameraPlugin.h"
+#include "MagicLeapCameraRunnable.h"
+#include "MagicLeapCameraTypes.h"
+
+DECLARE_LOG_CATEGORY_EXTERN(LogMagicLeapCamera, Verbose, All);
+
+class FMagicLeapCameraPlugin : public IMagicLeapCameraPlugin
+{
+public:
+	FMagicLeapCameraPlugin();
+	void StartupModule() override;
+	void ShutdownModule() override;
+	bool Tick(float DeltaTime) override;
+	bool CameraConnect(const FCameraConnect& ResultDelegate) override;
+	bool CameraDisconnect(const FCameraDisconnect& ResultDelegate) override;
+	int64 GetPreviewHandle() const override;
+
+	void IncUserCount();
+	void DecUserCount();
+	bool CaptureImageToFileAsync(const FCameraCaptureImgToFileMulti& ResultDelegate);
+	bool CaptureImageToTextureAsync(const FCameraCaptureImgToTextureMulti& ResultDelegate);
+	bool StartRecordingAsync(const FCameraStartRecordingMulti& ResultDelegate);
+	bool StopRecordingAsync(const FCameraStopRecordingMulti& ResultDelegate);
+	bool SetLogDelegate(const FCameraLogMessageMulti& LogDelegate);
+	bool IsCapturing() const;
+
+private:
+	bool TryPushNewCaptureTask(FCameraTask::EType InTaskType);
+
+	FMagicLeapAPISetup APISetup;
+	FTickerDelegate TickDelegate;
+	FDelegateHandle TickDelegateHandle;
+	uint32 UserCount;
+	FCameraRunnable* Runnable;
+	FCameraTask::EType CurrentTaskType;
+	FCameraTask::EType PrevTaskType;
+	FCameraConnect OnCameraConnect;
+	FCameraDisconnect OnCameraDisconnect;
+	FCameraCaptureImgToFileMulti OnCaptureImgToFile;
+	FCameraCaptureImgToTextureMulti OnCaptureImgToTexture;
+	FCameraStartRecordingMulti OnStartRecording;
+	FCameraStopRecordingMulti OnStopRecording;
+	FCameraLogMessageMulti OnLogMessage;
+};
+
+IMPLEMENT_MODULE(FMagicLeapCameraPlugin, MagicLeapCamera);
+
+#define GET_MAGIC_LEAP_CAMERA_PLUGIN() static_cast<FMagicLeapCameraPlugin*>(FModuleManager::Get().GetModule("MagicLeapCamera"))
diff --git a/Engine/Plugins/Lumin/MagicLeapCamera/Source/Private/MagicLeapCameraRunnable.cpp b/Engine/Plugins/Lumin/MagicLeapCamera/Source/Private/MagicLeapCameraRunnable.cpp
new file mode 100644
--- /dev/null
+++ b/Engine/Plugins/Lumin/MagicLeapCamera/Source/Private/MagicLeapCameraRunnable.cpp
@@ -0,0 +1,340 @@
+// Copyright 1998-2018 Epic Games, Inc. All Rights Reserved.
+
+#include "MagicLeapCameraRunnable.h"
+#include "MagicLeapCameraPlugin.h"
+#include "Engine/Texture2D.h"
+#include "Lumin/LuminPlatformFile.h"
+#if WITH_MLSDK
+ML_INCLUDES_START
+#include <ml_api.h>
+#include <ml_camera_metadata.h>
+#include <ml_camera_metadata_tags.h>
+#include <ml_media_error.h>
+ML_INCLUDES_END
+#endif //WITH_MLSDK
+
+FThreadSafeCounter64 FCameraRunnable::PreviewHandle = ML_INVALID_HANDLE;
+
+FCameraRunnable::FCameraRunnable()
+#if WITH_MLSDK
+	: FMagicLeapRunnable({ MLPrivilegeID_CameraCapture, MLPrivilegeID_AudioRecognizer, MLPrivilegeID_AudioCaptureMic, MLPrivilegeID_VoiceInput }, TEXT("FCameraRunnable"))
+#else
+	: FMagicLeapRunnable(TEXT("FCameraRunnable"))
+#endif //WITH_MLSDK
+	, bCameraConnected(false)
+	, bWasConnectedOnPause(false)
+	, ImgExtension(".jpeg")
+	, VidExtension(".mp4")
+{
+	IImageWrapperModule& ImageWrapperModule = FModuleManager::LoadModuleChecked<IImageWrapperModule>(FName("ImageWrapper"));
+	ImageWrapper = ImageWrapperModule.CreateImageWrapper(EImageFormat::JPEG);
+}
+
+void FCameraRunnable::Exit()
+{
+	TryDisconnect();
+}
+
+void FCameraRunnable::PushNewCaptureTask(FCameraTask::EType InTaskType)
+{
+#if WITH_MLSDK
+	if (!bCameraConnected && InTaskType != FCameraTask::EType::Connect)
+	{
+		FCameraTask ConnectTask;
+		ConnectTask.CaptureType = FCameraTask::EType::Connect;
+		PushNewTask(ConnectTask);
+	}
+	FCameraTask CaptureImgToFileTask;
+	CaptureImgToFileTask.CaptureType = InTaskType;
+	PushNewTask(CaptureImgToFileTask);
+#endif // WITH_MLSDK
+}
+
+bool FCameraRunnable::IsConnected() const
+{
+	return bCameraConnected;
+}
+
+void FCameraRunnable::Pause()
+{
+#if WITH_MLSDK
+	bWasConnectedOnPause = bCameraConnected;
+	// Cancel the current video recording (if one is active).
+	if (CurrentTask.CaptureType == FCameraTask::EType::StartVideoToFile)
+	{
+		StopRecordingVideo();
+		CurrentTask.bSuccess = false;
+		PushCompletedTask(CurrentTask);
+	}
+	// Cancel any incoming tasks.
+	CancelIncomingTasks();
+	// Disconnect camera if connected
+	TryDisconnect();
+#endif // WITH_MLSDK
+}
+
+void FCameraRunnable::Resume()
+{
+	if (bWasConnectedOnPause)
+	{
+		FCameraTask ConnectTask;
+		ConnectTask.CaptureType = FCameraTask::EType::Connect;
+		PushNewTask(ConnectTask);
+	}
+}
+
+#if WITH_MLSDK
+void FCameraRunnable::OnPreviewBufferAvailable(MLHandle Output, void *Data)
+{
+	(void)Data;
+	PreviewHandle.Set(static_cast<int64>(Output));
+}
+
+bool FCameraRunnable::ProcessCurrentTask()
+{
+	bool bSuccess = false;
+
+	switch (CurrentTask.CaptureType)
+	{
+	case FCameraTask::EType::None: bSuccess = false; checkf(false, TEXT("Invalid camera task encountered!")); break;
+	case FCameraTask::EType::Connect: bSuccess = bCameraConnected = TryConnect(); break;
+	case FCameraTask::EType::Disconnect: bSuccess = TryDisconnect(); break;
+	case FCameraTask::EType::ImageToFile: bSuccess = CaptureImageToFile(); break;
+	case FCameraTask::EType::ImageToTexture: bSuccess = CaptureImageToTexture(); break;
+	case FCameraTask::EType::StartVideoToFile: bSuccess = StartRecordingVideo(); break;
+	case FCameraTask::EType::StopVideoToFile: bSuccess = StopRecordingVideo(); break;
+	}
+
+	return bSuccess;
+}
+
+bool FCameraRunnable::TryConnect()
+{
+	//IncomingTasks.Peek(CurrentTask); // this is purely so that any log messages go to the correct requester
+	if (AppEventHandler.GetPrivilegeStatus(MLPrivilegeID_CameraCapture) != MagicLeap::EPrivilegeState::Granted)
+	{
+		Log(TEXT("Cannot connect to camera due to lack of privilege!"));
+		return false;
+	}
+	else
+	{
+		if (bPaused) return false;
+
+		MLResult Result = MLCameraConnect();
+
+		if (bPaused) return false;
+
+		if (Result != MLResult_Ok)
+		{
+			Log(FString::Printf(TEXT("MLCameraConnect failed with error %s!"), UTF8_TO_TCHAR(MLMediaResultGetString(Result))));
+			CancelIncomingTasks();
+			return false;
+		}
+		else
+		{
+			FMemory::Memset(&DeviceStatusCallbacks, 0, sizeof(MLCameraDeviceStatusCallbacks));
+			DeviceStatusCallbacks.on_preview_buffer_available = OnPreviewBufferAvailable;
+			Result = MLCameraSetDeviceStatusCallbacks(&DeviceStatusCallbacks, nullptr);
+			if (Result != MLResult_Ok)
+			{
+				Log(FString::Printf(TEXT("MLCameraSetDeviceStatusCallbacks failed with error %s!"), UTF8_TO_TCHAR(MLMediaResultGetString(Result))));
+			}
+		}
+	}
+
+	return true;
+}
+
+bool FCameraRunnable::TryDisconnect()
+{
+	if (bCameraConnected)
+	{
+		MLResult Result = MLCameraDisconnect();
+		if (Result != MLResult_Ok)
+		{
+			Log(FString::Printf(TEXT("MLCameraDisconnect failed with error %s!"), UTF8_TO_TCHAR(MLMediaResultGetString(Result))));
+		}
+		else
+		{
+			bCameraConnected = false;
+		}
+	}
+
+	return bCameraConnected == false;
+}
+
+bool FCameraRunnable::CaptureImageToFile()
+{
+	if (bPaused) return false;
+
+	Log(TEXT("Beginning capture image to file."));
+	MLHandle Handle = ML_INVALID_HANDLE;
+	MLResult Result = MLCameraPrepareCapture(MLCameraCaptureType_Image, &Handle);
+	if (Result != MLResult_Ok)
+	{
+		Log(FString::Printf(TEXT("MLCameraPrepareCapture failed with error %s!  Camera capture aborted!"), UTF8_TO_TCHAR(MLMediaResultGetString(Result))));
+		return false;
+	}
+
+	if (bPaused) return false;
+
+#if PLATFORM_LUMIN
+	IPlatformFile& PlatformFile = IPlatformFile::GetPlatformPhysical();
+	// This module is only for Lumin so this is fine for now.
+	FLuminPlatformFile* LuminPlatformFile = static_cast<FLuminPlatformFile*>(&PlatformFile);
+	UniqueFileName = LuminPlatformFile->ConvertToLuminPath(FPaths::CreateTempFilename(*FPaths::ProjectSavedDir(), TEXT("Img_"), *ImgExtension), true);
+#endif
+	Result = MLCameraCaptureImage(TCHAR_TO_UTF8(*UniqueFileName));
+	if (Result != MLResult_Ok)
+	{
+		Log(FString::Printf(TEXT("MLCameraCaptureImage failed with error %s!  Camera capture aborted!"), UTF8_TO_TCHAR(MLMediaResultGetString(Result))));
+		return false;
+	}
+
+	Log(FString::Printf(TEXT("Captured image to %s"), *UniqueFileName));
+	CurrentTask.FilePath = UniqueFileName;
+	return true;
+}
+
+bool FCameraRunnable::CaptureImageToTexture()
+{
+	if (bPaused) return false;
+
+	Log(TEXT("Beginning capture image to texture."));
+	MLCameraOutput* CameraOutput = nullptr;
+	MLHandle Handle = ML_INVALID_HANDLE;
+	MLResult Result = MLCameraPrepareCapture(MLCameraCaptureType_ImageRaw, &Handle);
+	if (Result != MLResult_Ok)
+	{
+		Log(FString::Printf(TEXT("MLCameraPrepareCapture failed with error %s!  Camera capture aborted!"), UTF8_TO_TCHAR(MLMediaResultGetString(Result))));
+		return false;
+	}
+
+	if (bPaused) return false;
+
+	Result = MLCameraCaptureImageRaw();
+	if (Result != MLResult_Ok)
+	{
+		Log(FString::Printf(TEXT("MLCameraCaptureImageRaw failed with error %s!  Camera capture aborted!"), UTF8_TO_TCHAR(MLMediaResultGetString(Result))));
+		return false;
+	}
+
+	if (bPaused) return false;
+
+	Result = MLCameraGetImageStream(&CameraOutput);
+	if (Result != MLResult_Ok)
+	{
+		Log(FString::Printf(TEXT("MLCameraGetImageStream failed with error %s!  Camera capture aborted!"), UTF8_TO_TCHAR(MLMediaResultGetString(Result))));
+		return false;
+	}
+
+	if (bPaused) return false;
+
+	if (CameraOutput->plane_count == 0)
+	{
+		Log(TEXT("Invalid plane_count!  Camera capture aborted!"));
+		return false;
+	}
+
+	MLCameraPlaneInfo& ImageInfo = CameraOutput->planes[0];
+	if (ImageWrapper.IsValid() && ImageWrapper->SetCompressed(ImageInfo.data, ImageInfo.size))
+	{
+		const TArray<uint8>* RawData = NULL;
+		if (ImageWrapper->GetRaw(ImageWrapper->GetFormat(), 8, RawData))
+		{
+			Log(FString::Printf(TEXT("ImageWrapper width=%d height=%d size=%d"), ImageWrapper->GetWidth(), ImageWrapper->GetHeight(), RawData->Num()));
+			UTexture2D* CaptureTexture = UTexture2D::CreateTransient(ImageWrapper->GetWidth(), ImageWrapper->GetHeight(), EPixelFormat::PF_R8G8B8A8);
+			CaptureTexture->AddToRoot();
+			FTexture2DMipMap& Mip = CaptureTexture->PlatformData->Mips[0];
+			void* Data = Mip.BulkData.Lock(LOCK_READ_WRITE);
+			FMemory::Memcpy(Data, RawData->GetData(), Mip.BulkData.GetBulkDataSize());
+			Mip.BulkData.Unlock();
+			CaptureTexture->UpdateResource();
+			CurrentTask.Texture = CaptureTexture;
+		}
+	}
+
+	return true;
+}
+
+bool FCameraRunnable::StartRecordingVideo()
+{
+	if (bPaused) return false;
+
+	Log(TEXT("Beginning capture video to file."));
+	MLHandle Handle = ML_INVALID_HANDLE;
+	MLResult Result = MLCameraPrepareCapture(MLCameraCaptureType_Video, &Handle);
+	if (Result != MLResult_Ok)
+	{
+		Log(FString::Printf(TEXT("MLCameraPrepareCapture failed with error %s!  Camera capture aborted!"), UTF8_TO_TCHAR(MLMediaResultGetString(Result))));
+		return false;
+	}
+
+	if (bPaused) return false;
+
+	if (AppEventHandler.GetPrivilegeStatus(MLPrivilegeID_AudioRecognizer) != MagicLeap::EPrivilegeState::Granted)
+	{
+		Log(TEXT("Cannot capture video due to lack of privilege!"));
+		return false;
+	}
+
+	if (bPaused) return false;
+
+	if (AppEventHandler.GetPrivilegeStatus(MLPrivilegeID_AudioCaptureMic) != MagicLeap::EPrivilegeState::Granted)
+	{
+		Log(TEXT("Cannot capture video due to lack of privilege!"));
+		return false;
+	}
+
+	if (bPaused) return false;
+
+	if (AppEventHandler.GetPrivilegeStatus(MLPrivilegeID_VoiceInput) != MagicLeap::EPrivilegeState::Granted)
+	{
+		Log(TEXT("Cannot capture video due to lack of privilege!"));
+		return false;
+	}
+
+	if (bPaused) return false;
+
+#if PLATFORM_LUMIN
+	IPlatformFile& PlatformFile = IPlatformFile::GetPlatformPhysical();
+	// This module is only for Lumin so this is fine for now.
+	FLuminPlatformFile* LuminPlatformFile = static_cast<FLuminPlatformFile*>(&PlatformFile);
+	UniqueFileName = LuminPlatformFile->ConvertToLuminPath(FPaths::CreateTempFilename(*FPaths::ProjectSavedDir(), TEXT("Vid_"), *VidExtension), true);
+#endif
+	Result = MLCameraCaptureVideoStart(TCHAR_TO_UTF8(*UniqueFileName));
+	if (Result != MLResult_Ok)
+	{
+		Log(FString::Printf(TEXT("MLCameraCaptureVideoStart failed with error %s!  Video capture aborted!"), UTF8_TO_TCHAR(MLMediaResultGetString(Result))));
+		return false;
+	}
+
+	return true;
+}
+
+bool FCameraRunnable::StopRecordingVideo()
+{
+	MLResult Result = MLCameraCaptureVideoStop();
+	if (Result != MLResult_Ok)
+	{
+		Log(FString::Printf(TEXT("MLCameraCaptureVideoStop failed with error %s!"), UTF8_TO_TCHAR(MLMediaResultGetString(Result))));
+		return false;
+	}
+	else
+	{
+		Log(FString::Printf(TEXT("Captured video to %s"), *UniqueFileName));
+	}
+
+	CurrentTask.FilePath = UniqueFileName;
+	return true;
+}
+
+void FCameraRunnable::Log(const FString& Info)
+{
+	FCameraTask LogTask;
+	LogTask.CaptureType = FCameraTask::EType::Log;
+	LogTask.Log = Info;
+	PushCompletedTask(LogTask);
+	UE_LOG(LogMagicLeapCamera, Log, TEXT("%s"), *Info);
+}
+#endif //WITH_MLSDK
diff --git a/Engine/Plugins/Lumin/MagicLeapCamera/Source/Private/MagicLeapCameraRunnable.h b/Engine/Plugins/Lumin/MagicLeapCamera/Source/Private/MagicLeapCameraRunnable.h
new file mode 100644
--- /dev/null
+++ b/Engine/Plugins/Lumin/MagicLeapCamera/Source/Private/MagicLeapCameraRunnable.h
@@ -0,0 +1,75 @@
+// Copyright 1998-2018 Epic Games, Inc. All Rights Reserved.
+
+#pragma once
+
+#include "MagicLeapRunnable.h"
+#include "IImageWrapper.h"
+#include "IImageWrapperModule.h"
+#if WITH_MLSDK
+ML_INCLUDES_START
+#include <ml_camera.h>
+ML_INCLUDES_END
+#endif //WITH_MLSDK
+
+struct FCameraTask : public FMagicLeapTask
+{
+	enum class EType : uint32
+	{
+		None,
+		Connect,
+		Disconnect,
+		ImageToFile,
+		ImageToTexture,
+		StartVideoToFile,
+		StopVideoToFile,
+		Log,
+	};
+
+	EType CaptureType;
+	FString FilePath;
+	FString Log;
+	class UTexture2D* Texture;
+
+	FCameraTask()
+		: CaptureType(EType::None)
+		, Texture(nullptr)
+	{
+	}
+};
+
+class FCameraRunnable : public FMagicLeapRunnable<FCameraTask>
+{
+public:
+	FCameraRunnable();
+	void Exit() override;
+	void PushNewCaptureTask(FCameraTask::EType InTaskType);
+	bool IsConnected() const;
+
+	static FThreadSafeCounter64 PreviewHandle;
+
+protected:
+	void Pause() override;
+	void Resume() override;
+
+private:
+#if WITH_MLSDK
+	static void OnPreviewBufferAvailable(MLHandle Output, void *Data);
+	bool ProcessCurrentTask() override;
+	bool TryConnect();
+	bool TryDisconnect();
+	bool CaptureImageToFile();
+	bool CaptureImageToTexture();
+	bool StartRecordingVideo();
+	bool StopRecordingVideo();
+	void Log(const FString& InLogMsg);
+
+	MLCameraDeviceStatusCallbacks DeviceStatusCallbacks;
+#endif //WITH_MLSDK
+	FCameraTask::EType CurrentTaskType;
+	FThreadSafeBool bCameraConnected;
+	bool bWasConnectedOnPause;
+	const FString ImgExtension;
+	const FString VidExtension;
+	FString UniqueFileName;
+	TSharedPtr<IImageWrapper> ImageWrapper;
+};
diff --git a/Engine/Plugins/Lumin/MagicLeapCamera/Source/Public/IMagicLeapCameraPlugin.h b/Engine/Plugins/Lumin/MagicLeapCamera/Source/Public/IMagicLeapCameraPlugin.h
new file mode 100644
--- /dev/null
+++ b/Engine/Plugins/Lumin/MagicLeapCamera/Source/Public/IMagicLeapCameraPlugin.h
@@ -0,0 +1,41 @@
+// Copyright 1998-2018 Epic Games, Inc. All Rights Reserved.
+
+#pragma once
+
+#include "Modules/ModuleManager.h"
+#include "Containers/Ticker.h"
+#include "MagicLeapCameraTypes.h"
+
+/**
+ * The public interface to this module.  In most cases, this interface is only public to sibling modules
+ * within this plugin.
+ */
+class IMagicLeapCameraPlugin : public IModuleInterface
+{
+public:
+	/**
+	 * Singleton-like access to this module's interface.  This is just for convenience!
+	 * Beware of calling this during the shutdown phase, though.  Your module might have been unloaded already.
+	 *
+	 * @return Returns singleton instance, loading the module on demand if needed
+	 */
+	static inline IMagicLeapCameraPlugin& Get()
+	{
+		return FModuleManager::LoadModuleChecked<IMagicLeapCameraPlugin>("MagicLeapCamera");
+	}
+
+	/**
+	 * Checks to see if this module is loaded and ready.  It is only valid to call Get() if IsAvailable() returns true.
+	 *
+	 * @return True if the module is loaded and ready to use
+	 */
+	static inline bool IsAvailable()
+	{
+		return FModuleManager::Get().IsModuleLoaded("MagicLeapCamera");
+	}
+
+	virtual bool Tick(float DeltaTime) = 0;
+	virtual bool CameraConnect(const FCameraConnect& ResultDelegate) { return true; }
+	virtual bool CameraDisconnect(const FCameraDisconnect& ResultDelegate) { return true; }
+	virtual int64 GetPreviewHandle() const = 0;
+};
