Removed spamming log line "Mesh info timestamp".

diff --git a/Engine/Plugins/Lumin/MagicLeap/Source/MagicLeap/Private/MeshTrackerComponent.cpp b/Engine/Plugins/Lumin/MagicLeap/Source/MagicLeap/Private/MeshTrackerComponent.cpp
--- a/Engine/Plugins/Lumin/MagicLeap/Source/MagicLeap/Private/MeshTrackerComponent.cpp
+++ b/Engine/Plugins/Lumin/MagicLeap/Source/MagicLeap/Private/MeshTrackerComponent.cpp
@@ -660,7 +660,6 @@ bool UMeshTrackerComponent::GetMeshInfoResult()
 			Impl->LastestMeshInfo.BlockData.Empty(MeshInfo.data_count);
 			Impl->LastestMeshInfo.BlockData.AddUninitialized(MeshInfo.data_count);
 			Impl->LastestMeshInfo.Timestamp = FTimespan::FromMicroseconds(MeshInfo.timestamp / 1000.0);
-			UE_LOG(LogMagicLeap, Warning, TEXT("Mesh info timestamp = %f"), Impl->LastestMeshInfo.Timestamp.GetTotalMilliseconds());
 
 			for (uint32_t MeshInfoIndex = 0; MeshInfoIndex < MeshInfo.data_count; ++ MeshInfoIndex)
 			{
