Renamed FMeshInfo to FTrackingMeshInfo to resolve a name conflict with a ustruct named the same thing.

diff --git a/Engine/Plugins/Lumin/MagicLeap/Config/BaseMagicLeap.ini b/Engine/Plugins/Lumin/MagicLeap/Config/BaseMagicLeap.ini
new file mode 100644
--- /dev/null
+++ b/Engine/Plugins/Lumin/MagicLeap/Config/BaseMagicLeap.ini
@@ -0,0 +1,3 @@
+[CoreRedirects]
+
++StructRedirects=(OldName="/Script/MagicLeap.MeshInfo", NewName="/Script/MagicLeap.MLTrackingMeshInfo")
diff --git a/Engine/Plugins/Lumin/MagicLeap/Source/MagicLeap/Classes/MeshBlockSelectorInterface.h b/Engine/Plugins/Lumin/MagicLeap/Source/MagicLeap/Classes/MeshBlockSelectorInterface.h
--- a/Engine/Plugins/Lumin/MagicLeap/Source/MagicLeap/Classes/MeshBlockSelectorInterface.h
+++ b/Engine/Plugins/Lumin/MagicLeap/Source/MagicLeap/Classes/MeshBlockSelectorInterface.h
@@ -23,5 +23,5 @@ public:
 	 * @param RequestedMesh output array containing details of blocks for whom the mesh should be requested.
 	 */
 	UFUNCTION(BlueprintCallable, BlueprintNativeEvent, Category="Meshing|MagicLeap")
-	void SelectMeshBlocks(const FMeshInfo& NewMeshInfo, TArray<FMeshBlockRequest>& RequestedMesh);
+	void SelectMeshBlocks(const FMLTrackingMeshInfo& NewMeshInfo, TArray<FMeshBlockRequest>& RequestedMesh);
 };
diff --git a/Engine/Plugins/Lumin/MagicLeap/Source/MagicLeap/Classes/MeshTrackerComponent.h b/Engine/Plugins/Lumin/MagicLeap/Source/MagicLeap/Classes/MeshTrackerComponent.h
--- a/Engine/Plugins/Lumin/MagicLeap/Source/MagicLeap/Classes/MeshTrackerComponent.h
+++ b/Engine/Plugins/Lumin/MagicLeap/Source/MagicLeap/Classes/MeshTrackerComponent.h
@@ -54,8 +54,8 @@ public:
 
 	/** IMeshBlockSelectorInterface */
 	UFUNCTION(BlueprintCallable, BlueprintNativeEvent, Category="Meshing|MagicLeap")
-	void SelectMeshBlocks(const FMeshInfo& NewMeshInfo, TArray<FMeshBlockRequest>& RequestedMesh);
-	virtual void SelectMeshBlocks_Implementation(const FMeshInfo& NewMeshInfo, TArray<FMeshBlockRequest>& RequestedMesh) override;
+	void SelectMeshBlocks(const FMLTrackingMeshInfo& NewMeshInfo, TArray<FMeshBlockRequest>& RequestedMesh);
+	virtual void SelectMeshBlocks_Implementation(const FMLTrackingMeshInfo& NewMeshInfo, TArray<FMeshBlockRequest>& RequestedMesh) override;
 
 	/**
 		Delegate used by OnMeshUpdated().
diff --git a/Engine/Plugins/Lumin/MagicLeap/Source/MagicLeap/Classes/MeshTrackerTypes.h b/Engine/Plugins/Lumin/MagicLeap/Source/MagicLeap/Classes/MeshTrackerTypes.h
--- a/Engine/Plugins/Lumin/MagicLeap/Source/MagicLeap/Classes/MeshTrackerTypes.h
+++ b/Engine/Plugins/Lumin/MagicLeap/Source/MagicLeap/Classes/MeshTrackerTypes.h
@@ -91,7 +91,7 @@ public:
 
 /** Response structure for the mesh block info. */
 USTRUCT(BlueprintType)
-struct FMeshInfo
+struct FMLTrackingMeshInfo
 {
 	GENERATED_BODY()
 
diff --git a/Engine/Plugins/Lumin/MagicLeap/Source/MagicLeap/Private/MeshTrackerComponent.cpp b/Engine/Plugins/Lumin/MagicLeap/Source/MagicLeap/Private/MeshTrackerComponent.cpp
--- a/Engine/Plugins/Lumin/MagicLeap/Source/MagicLeap/Private/MeshTrackerComponent.cpp
+++ b/Engine/Plugins/Lumin/MagicLeap/Source/MagicLeap/Private/MeshTrackerComponent.cpp
@@ -184,7 +184,7 @@ public:
 
 #endif //WITH_MLSDK
 
-	FMeshInfo LastestMeshInfo;
+	FMLTrackingMeshInfo LastestMeshInfo;
 	TArray<FMeshBlockRequest> UEMeshBlockRequests;
 	TScriptInterface<IMeshBlockSelectorInterface> BlockSelector;
 
@@ -534,7 +534,7 @@ void UMeshTrackerComponent::TickComponent(float DeltaTime, enum ELevelTick TickT
 #endif //WITH_MLSDK
 }
 
-void UMeshTrackerComponent::SelectMeshBlocks_Implementation(const FMeshInfo& NewMeshInfo, TArray<FMeshBlockRequest>& RequestedMesh)
+void UMeshTrackerComponent::SelectMeshBlocks_Implementation(const FMLTrackingMeshInfo& NewMeshInfo, TArray<FMeshBlockRequest>& RequestedMesh)
 {
 	for (const FMeshBlockInfo& BlockInfo : NewMeshInfo.BlockData)
 	{
