UEVR-1360 Spatial anchor support - planes also - Lumin MagicLeap
-Refactoring so Lumin uses generic AR system code, rather than only its own custom implementations (previous implementation will still work).
-Planes detection works.  Inner planes are not yet supported.
-ARLineCheck works with extents or bounds.  There is no 'ground plane' check now.
-Pins work.
-A new 3d line check (where you specify start and end point for the test) was added and is supported on Lumin and ARCore.  There is no ARKit support for that now.  The 'screen space' line check casts a ray from 10cm below hmd forward on lumin.

diff --git a/Engine/Plugins/Lumin/MagicLeap/MagicLeap.uplugin b/Engine/Plugins/Lumin/MagicLeap/MagicLeap.uplugin
--- a/Engine/Plugins/Lumin/MagicLeap/MagicLeap.uplugin
+++ b/Engine/Plugins/Lumin/MagicLeap/MagicLeap.uplugin
@@ -30,6 +30,16 @@
 				"Lumin"
 			]
 		},
+		{
+			"Name" : "MagicLeapAR",
+			"Type" : "Runtime",
+			"WhitelistPlatforms" :
+			[
+				"Win64",
+				"Mac",
+				"Lumin"
+			]
+		},
 		{
 			"Name" : "MagicLeapAudio",
 			"Type" : "Runtime",
diff --git a/Engine/Plugins/Lumin/MagicLeap/Source/MagicLeap/Classes/PlanesComponent.h b/Engine/Plugins/Lumin/MagicLeap/Source/MagicLeap/Classes/PlanesComponent.h
--- a/Engine/Plugins/Lumin/MagicLeap/Source/MagicLeap/Classes/PlanesComponent.h
+++ b/Engine/Plugins/Lumin/MagicLeap/Source/MagicLeap/Classes/PlanesComponent.h
@@ -2,40 +2,26 @@
 
 #pragma once
 
+#if WITH_MLSDK
+#include "ml_planes.h"
+#endif //WITH_MLSDK
+
 #include "Components/SceneComponent.h"
+#include "LuminARTypes.h"
 #include "PlanesComponent.generated.h"
 
-/** Control flags for plane queries. */
-UENUM(BlueprintType)
-enum class EPlaneQueryFlags : uint8
-{
-	/** Include planes whose normal is perpendicular to gravity. */
-	Vertical,
-
-	/** Include planes whose normal is parallel to gravity. */
-	Horizontal,
+#if WITH_MLSDK
 
-	/** Include planes with arbitrary normals. */
-	Arbitrary,
+MLPlanesQueryFlags UnrealToMLPlanesQueryFlagMap(EPlaneQueryFlags QueryFlag);
 
-	/** If set, non-horizontal planes will be aligned perpendicular to gravity. */
-	OrientToGravity,
+EPlaneQueryFlags MLToUnrealPlanesQueryFlagMap(MLPlanesQueryFlags QueryFlag);
 
-	/** If set, inner planes will be returned; if not set, outer planes will be returned. */
-	PreferInner,
+MLPlanesQueryFlags UnrealToMLPlanesQueryFlags(const TArray<EPlaneQueryFlags>& QueryFlags);
 
-	/** If set, holes in planar surfaces will be ignored. */
-	IgnoreHoles,
+void MLToUnrealPlanesQueryFlags(uint32 QueryFlags, TArray<EPlaneQueryFlags>& OutPlaneFlags);
 
-	/** If set, include planes semantically tagged as ceiling. */
-	Ceiling,
+#endif //WITH_MLSDK
 
-	/** If set, include planes semantically tagged as floor. */
-	Floor,
-
-	/** If set, include planes semantically tagged as wall. */
-	Wall
-};
 
 /** Represents a plane returned from the ML-API. */
 USTRUCT(BlueprintType)
@@ -64,9 +50,16 @@ public:
 	UPROPERTY(EditAnywhere, BlueprintReadWrite, Category = "Planes|MagicLeap")
 	TArray<EPlaneQueryFlags> PlaneFlags;
 
+	/** The Boundary of the plane in plane local space.*/
+	UPROPERTY(EditAnywhere, BlueprintReadWrite, Category = "Planes|MagicLeap")
+	TArray<FVector> BoundaryPolygon;
+
 	/** ID of the plane result. This ID is persistent across queries */
 	UPROPERTY(BlueprintReadOnly, Category = "Planes|MagicLeap")
 	FGuid ID;
+
+	/** ID of the plane result. This ID is persistent across queries */
+	uint64 ID_64;
 };
 
 /** 
diff --git a/Engine/Plugins/Lumin/MagicLeap/Source/MagicLeap/MagicLeap.Build.cs b/Engine/Plugins/Lumin/MagicLeap/Source/MagicLeap/MagicLeap.Build.cs
--- a/Engine/Plugins/Lumin/MagicLeap/Source/MagicLeap/MagicLeap.Build.cs
+++ b/Engine/Plugins/Lumin/MagicLeap/Source/MagicLeap/MagicLeap.Build.cs
@@ -51,7 +51,8 @@ namespace UnrealBuildTool.Rules
 					// Public headers of MagicLeapHelperVulkan are protected against Mac so this is fine here.
 					"MagicLeapHelperVulkan",
 					"LuminRuntimeSettings",
-					"MagicLeapSecureStorage"
+					"MagicLeapSecureStorage",
+					"MagicLeapAR",
 				}
 			);
 
diff --git a/Engine/Plugins/Lumin/MagicLeap/Source/MagicLeap/Private/MagicLeapHMD.cpp b/Engine/Plugins/Lumin/MagicLeap/Source/MagicLeap/Private/MagicLeapHMD.cpp
--- a/Engine/Plugins/Lumin/MagicLeap/Source/MagicLeap/Private/MagicLeapHMD.cpp
+++ b/Engine/Plugins/Lumin/MagicLeap/Source/MagicLeap/Private/MagicLeapHMD.cpp
@@ -43,6 +43,7 @@
 #include "VulkanRHIBridge.h"
 #endif
 
+#include "LuminARModule.h"
 #include "MagicLeapPluginUtil.h" // for ML_INCLUDES_START/END
 
 #if WITH_MLSDK
@@ -174,9 +175,13 @@ public:
 		else
 		{
 			//initialize AR system
-			TSharedPtr<IARSystemSupport, ESPMode::ThreadSafe> ARImplementation;
+			auto ARModule = FModuleManager::LoadModulePtr<ILuminARModule>("MagicLeapAR");
+			check(ARModule);
+			ARImplementation = ARModule->CreateARImplementation();
 			LocalHMD = MakeShared<FMagicLeapHMD, ESPMode::ThreadSafe>(this, ARImplementation.Get(), bIsVDZIEnabled, bUseVulkanForZI);
 			HMD = LocalHMD;
+			ARModule->ConnectARImplementationToXRSystem(LocalHMD.Get());
+			ARModule->InitializeARImplementation();
 		}
 #endif
 		if (LocalHMD.IsValid() && !LocalHMD->IsInitialized())
@@ -307,6 +312,7 @@ private:
 	bool bUseVulkanForZI;
 	FMagicLeapAPISetup APISetup;
 	TWeakPtr<FMagicLeapHMD, ESPMode::ThreadSafe> HMD;
+	TSharedPtr<IARSystemSupport, ESPMode::ThreadSafe> ARImplementation;
 	TSharedPtr<FMagicLeapVulkanExtensions, ESPMode::ThreadSafe> VulkanExtensions;
 	TSet<IMagicLeapInputDevice*> InputDevices;
 };
diff --git a/Engine/Plugins/Lumin/MagicLeap/Source/MagicLeap/Private/PlanesComponent.cpp b/Engine/Plugins/Lumin/MagicLeap/Source/MagicLeap/Private/PlanesComponent.cpp
--- a/Engine/Plugins/Lumin/MagicLeap/Source/MagicLeap/Private/PlanesComponent.cpp
+++ b/Engine/Plugins/Lumin/MagicLeap/Source/MagicLeap/Private/PlanesComponent.cpp
@@ -12,9 +12,6 @@
 #include "Editor.h"
 #endif
 
-#if WITH_MLSDK
-#include "ml_planes.h"
-#endif //WITH_MLSDK
 
 class FPlanesTrackerImpl
 {
@@ -36,7 +33,6 @@ public:
 #if WITH_MLSDK
 		if (!MLHandleIsValid(Tracker))
 		{
-			UE_LOG(LogMagicLeap, Display, TEXT("Creating Planes Tracker"));
 			MLResult CreateResult = MLPlanesCreate(&Tracker);
 
 			if (CreateResult != MLResult_Ok || !MLHandleIsValid(Tracker))
diff --git a/Engine/Plugins/Lumin/MagicLeap/Source/MagicLeapAR/MagicLeapAR.Build.cs b/Engine/Plugins/Lumin/MagicLeap/Source/MagicLeapAR/MagicLeapAR.Build.cs
new file mode 100644
--- /dev/null
+++ b/Engine/Plugins/Lumin/MagicLeap/Source/MagicLeapAR/MagicLeapAR.Build.cs
@@ -0,0 +1,43 @@
+// Copyright 1998-2018 Epic Games, Inc. All Rights Reserved.
+
+using System.IO;
+
+namespace UnrealBuildTool.Rules
+{
+	public class MagicLeapAR : ModuleRules
+	{
+		public MagicLeapAR(ReadOnlyTargetRules Target) : base(Target)
+		{
+			//JB - Remove
+            OptimizeCode = CodeOptimization.Never;
+
+            PrivateIncludePaths.AddRange(
+				new string[] {
+					"../../../../Source/Runtime/Renderer/Private",
+					// ... add other private include paths required here ...
+				}
+			);
+			// This is not ideal but needs to be done in order to expose the private MagicLeapHMD header to this module.
+			//PrivateIncludePaths.Add(Path.Combine(new string[] { ModuleDirectory, "../../../../Lumin/MagicLeap", "Private" }));
+
+			PrivateDependencyModuleNames.AddRange(
+				new string[]
+				{
+					"HeadMountedDisplay",
+					"Core",
+					"CoreUObject",
+					"Engine",
+					"RenderCore",
+					"LuminRuntimeSettings",
+					"MagicLeap",
+					"MLSDK",
+					"HeadMountedDisplay",
+					"AugmentedReality"
+				}
+			);
+
+			// This is not ideal but needs to be done in order to expose the private MagicLeapHMD header to this module.
+			PrivateIncludePaths.Add(Path.Combine(new string[] { ModuleDirectory, "..", "MagicLeap", "Private" }));
+		}
+	}
+}
diff --git a/Engine/Plugins/Lumin/MagicLeap/Source/MagicLeapAR/Private/LuminARAPI.cpp b/Engine/Plugins/Lumin/MagicLeap/Source/MagicLeapAR/Private/LuminARAPI.cpp
new file mode 100644
--- /dev/null
+++ b/Engine/Plugins/Lumin/MagicLeap/Source/MagicLeapAR/Private/LuminARAPI.cpp
@@ -0,0 +1,949 @@
+// Copyright 1998-2018 Epic Games, Inc. All Rights Reserved.
+
+#include "LuminARAPI.h"
+
+#include "Misc/EngineVersion.h"
+#include "DrawDebugHelpers.h"
+#include "Templates/Casts.h"
+
+#include "IMagicLeapPlugin.h"
+#include "MagicLeapPrivileges.h"
+#include "MagicLeapHMD.h"
+#include "AppFramework.h"
+#include "HeadMountedDisplayFunctionLibrary.h"
+
+#if WITH_MLSDK
+#include "ml_planes.h"
+#endif //WITH_MLSDK
+
+namespace
+{
+#if PLATFORM_LUMIN
+	static const FMatrix LuminARToUnrealTransform = FMatrix(
+		FPlane(0.0f, 0.0f, -1.0f, 0.0f),
+		FPlane(1.0f, 0.0f, 0.0f, 0.0f),
+		FPlane(0.0f, 1.0f, 0.0f, 0.0f),
+		FPlane(0.0f, 0.0f, 0.0f, 1.0f));
+
+	static const FMatrix LuminARToUnrealTransformInverse = LuminARToUnrealTransform.InverseFast();
+
+	FTransform LuminARPoseToUnrealTransform(const ArPose& InPose, const FLuminARSession* Session, float WorldToMeterScale)
+	{
+		FTransform LuminArPoseTransform;
+		LuminArPoseTransform.SetTranslation(InPose.Pos);
+		LuminArPoseTransform.SetRotation(InPose.Quat);
+
+		FMatrix LuminARPoseMatrix = LuminArPoseTransform.ToMatrixNoScale();
+		FTransform Result = FTransform(LuminARToUnrealTransform * LuminARPoseMatrix * LuminARToUnrealTransformInverse);
+		Result.SetLocation(Result.GetLocation() * WorldToMeterScale);
+
+		return Result;
+	}
+
+	void UnrealTransformToLuminARPose(const FTransform& UnrealTransform, const FLuminARSession* Session, ArPose& OutPose, float WorldToMeterScale)
+	{
+		FMatrix UnrealPoseMatrix = UnrealTransform.ToMatrixNoScale();
+		UnrealPoseMatrix.SetOrigin(UnrealPoseMatrix.GetOrigin() / WorldToMeterScale);
+		FMatrix LuminARPoseMatrix = LuminARToUnrealTransformInverse * UnrealPoseMatrix * LuminARToUnrealTransform;
+
+		FVector ArPosePosition = LuminARPoseMatrix.GetOrigin();
+		FQuat ArPoseRotation = LuminARPoseMatrix.ToQuat();
+
+		OutPose.Pos = LuminARPoseMatrix.GetOrigin();
+		OutPose.Quat = LuminARPoseMatrix.ToQuat();
+	}
+#endif
+	inline bool CheckIsSessionValid(FString TypeName, const TWeakPtr<FLuminARSession>& SessionPtr)
+	{
+		return true;
+	}
+}
+
+
+/****************************************/
+/*         FLuminARSession         */
+/****************************************/
+FLuminARSession::FLuminARSession()
+	: LatestFrame(nullptr)
+	, UObjectManager(nullptr)
+	, CachedWorldToMeterScale(100.0f)
+	, FrameNumber(0)
+
+{
+	// Create Lumin ARSession handle.
+	LatestFrame = new FLuminARFrame(this);
+#if PLATFORM_LUMIN
+	InitTracker();
+	LatestFrame->Init();
+#endif
+}
+
+FLuminARSession::~FLuminARSession()
+{
+	for (UARPin* Anchor : UObjectManager->AllAnchors)
+	{
+		Anchor->OnTrackingStateChanged(EARTrackingState::StoppedTracking);
+	}
+
+	delete LatestFrame;
+
+#if PLATFORM_LUMIN
+	DestroyTracker();
+#endif
+}
+
+
+ULuminARUObjectManager* FLuminARSession::GetUObjectManager()
+{
+	return UObjectManager;
+}
+
+float FLuminARSession::GetWorldToMeterScale()
+{
+	return CachedWorldToMeterScale;
+}
+
+
+#if PLATFORM_ANDROID
+MLHandle FLuminARSession::GetPlaneTrackerHandle()
+{
+	return PlaneTrackerHandle;
+}
+#endif
+
+ELuminARAPIStatus FLuminARSession::Resume()
+{
+	ELuminARAPIStatus ResumeStatus = ELuminARAPIStatus::AR_SUCCESS;
+#if PLATFORM_LUMIN
+	//InitTracker();
+#endif
+	return ResumeStatus;
+}
+
+ELuminARAPIStatus FLuminARSession::Pause()
+{
+	ELuminARAPIStatus PauseStatus = ELuminARAPIStatus::AR_SUCCESS;
+#if PLATFORM_LUMIN
+	//DestroyTracker();
+#endif
+
+	for (UARPin* Anchor : UObjectManager->AllAnchors)
+	{
+		Anchor->OnTrackingStateChanged(EARTrackingState::NotTracking);
+	}
+
+	return PauseStatus;
+}
+
+ELuminARAPIStatus FLuminARSession::Update(float WorldToMeterScale)
+{
+	ELuminARAPIStatus UpdateStatus = ELuminARAPIStatus::AR_SUCCESS;
+
+	CachedWorldToMeterScale = WorldToMeterScale;
+	int64 LastFrameTimestamp = LatestFrame->GetCameraTimestamp();
+	LatestFrame->Update(WorldToMeterScale);
+	if (LastFrameTimestamp != LatestFrame->GetCameraTimestamp())
+	{
+		FrameNumber++;
+	}
+
+	return UpdateStatus;
+}
+
+const FLuminARFrame* FLuminARSession::GetLatestFrame()
+{
+	return LatestFrame;
+}
+
+uint32 FLuminARSession::GetFrameNum()
+{
+	return FrameNumber;
+}
+
+
+// Anchors and Planes
+ELuminARAPIStatus FLuminARSession::CreateARAnchor(const FTransform& TransfromInTrackingSpace, UARTrackedGeometry* TrackedGeometry, USceneComponent* ComponentToPin, FName InDebugName, UARPin*& OutAnchor)
+{
+	ELuminARAPIStatus AnchorCreateStatus = ELuminARAPIStatus::AR_SUCCESS;
+	OutAnchor = nullptr;
+
+#if PLATFORM_LUMIN
+	TSharedPtr<LuminArAnchor> NewLuminArAnchor;
+	ArPose Pose;
+	UnrealTransformToLuminARPose(TransfromInTrackingSpace, this, Pose, CachedWorldToMeterScale);
+	if (TrackedGeometry == nullptr)
+	{
+		NewLuminArAnchor = MakeShared<LuminArAnchor>(Pose, ML_INVALID_HANDLE);
+	}
+	else
+	{
+		ensure(TrackedGeometry->GetNativeResource() != nullptr);
+		MLHandle ParentHandle = reinterpret_cast<FLuminARTrackableResource*>(TrackedGeometry->GetNativeResource())->GetNativeHandle();
+		ensure(ParentHandle != ML_INVALID_HANDLE);
+		NewLuminArAnchor = MakeShared<LuminArAnchor>(Pose, ParentHandle);
+	}
+
+	if (AnchorCreateStatus == ELuminARAPIStatus::AR_SUCCESS)
+	{
+		OutAnchor = NewObject<UARPin>();
+		OutAnchor->InitARPin(GetARSystem(), ComponentToPin, TransfromInTrackingSpace, TrackedGeometry, InDebugName);
+		OutAnchor->SetNativeResource(reinterpret_cast<void*>(NewLuminArAnchor.Get()));
+
+		UObjectManager->HandleToLuminAnchorMap.Add(NewLuminArAnchor->Handle, NewLuminArAnchor);
+		UObjectManager->AllAnchors.Add(OutAnchor);
+		UObjectManager->HandleToAnchorMap.Add(NewLuminArAnchor->Handle, OutAnchor);
+	}
+#endif
+	return AnchorCreateStatus;
+}
+
+void FLuminARSession::DetachAnchor(UARPin* Anchor)
+{
+	if (!UObjectManager->AllAnchors.Contains(Anchor))
+	{
+		return;
+	}
+
+#if PLATFORM_LUMIN
+	LuminArAnchor* NativeResource = reinterpret_cast<LuminArAnchor*>(Anchor->GetNativeResource());
+	check(NativeResource);
+	NativeResource->Detach();
+	Anchor->SetNativeResource(nullptr);
+
+	Anchor->OnTrackingStateChanged(EARTrackingState::StoppedTracking);
+
+	UObjectManager->HandleToAnchorMap.Remove(NativeResource->Handle);
+	UObjectManager->HandleToLuminAnchorMap.Remove(NativeResource->Handle);
+#endif
+	UObjectManager->AllAnchors.Remove(Anchor);
+}
+
+void FLuminARSession::GetAllAnchors(TArray<UARPin*>& OutAnchors) const
+{
+	OutAnchors = UObjectManager->AllAnchors;
+}
+
+void FLuminARSession::InitTracker()
+{
+#if PLATFORM_LUMIN
+	if (!MLHandleIsValid(PlaneTrackerHandle))
+	{
+		MLResult PlaneCreateResult = MLPlanesCreate(&PlaneTrackerHandle);
+		if (PlaneCreateResult != MLResult_Ok)
+		{
+			UE_LOG(LogLuminARAPI, Warning, TEXT("Failed to create Plane Tracker for Lumin AR Session Result:%i"), PlaneCreateResult);
+		}
+	}
+	else
+	{
+		UE_LOG(LogLuminARAPI, Warning, TEXT("Tracker already exists"));
+	}
+#endif
+}
+
+void FLuminARSession::DestroyTracker()
+{
+#if PLATFORM_LUMIN
+	if (MLHandleIsValid(PlaneTrackerHandle))
+	{
+		MLResult Result = MLPlanesDestroy(PlaneTrackerHandle);
+		if (Result != MLResult_Ok)
+		{
+			UE_LOG(LogLuminARAPI, Warning, TEXT("Failed to destroy Plane Tracker for Lumin AR Session Result:%i"), Result);
+		}
+
+		PlaneTrackerHandle = ML_INVALID_HANDLE;
+	}
+#endif //WITH_MLSDK
+}
+
+void FLuminARSession::AddReferencedObjects(FReferenceCollector& Collector)
+{
+	if (UObjectManager)
+	{
+		Collector.AddReferencedObject(UObjectManager);
+	}
+}
+
+/****************************************/
+/*         FLuminARFrame           */
+/****************************************/
+FLuminARFrame::FLuminARFrame(FLuminARSession* InSession)
+	: Session(InSession)
+	//, LatestCameraPose(FTransform::Identity)
+	, LatestCameraTimestamp(0)
+	, LatestCameraTrackingState(ELuminARTrackingState::StoppedTracking)
+	, LatestARPlaneQueryStatus(ELuminARPlaneQueryStatus::Unknown)
+#if PLATFORM_LUMIN
+	, PlaneTrackerHandle(ML_INVALID_HANDLE)
+#endif
+{
+}
+
+FLuminARFrame::~FLuminARFrame()
+{
+}
+
+void FLuminARFrame::Init()
+{
+#if PLATFORM_LUMIN
+	if (Session->GetPlaneTrackerHandle())
+	{
+		PlaneTrackerHandle = Session->GetPlaneTrackerHandle();
+}
+#endif
+}
+
+
+void FLuminARFrame::Update(float WorldToMeterScale)
+{
+#if PLATFORM_LUMIN
+	if (!MLHandleIsValid(PlaneTrackerHandle))
+	{
+		LatestCameraTrackingState = ELuminARTrackingState::NotTracking;
+		return;
+	}
+
+	// Update trackable that is cached in Unreal
+	StartPlaneQuery();
+	ProcessPlaneQuery();
+
+	switch (LatestARPlaneQueryStatus)
+	{
+	case ELuminARPlaneQueryStatus::Unknown:
+		return;
+	case ELuminARPlaneQueryStatus::Success:
+		{
+			//TODO - is pose needed?
+			//ArPose Pose;
+			//LatestCameraPose = Pose;
+
+			LatestCameraTimestamp = FPlatformTime::Seconds();
+			LatestCameraTrackingState = ELuminARTrackingState::Tracking;
+		}
+		break;
+	case ELuminARPlaneQueryStatus::Fail:
+		{
+			LatestCameraTrackingState = ELuminARTrackingState::NotTracking;
+		}
+		return;
+	default:
+		check(false);
+		return;
+	}
+
+	// Update Anchors
+	UpdatedAnchors.Empty();
+	for (auto HandleToAnchorMapPair : Session->GetUObjectManager()->HandleToAnchorMap)
+	{
+		const MLHandle AnchorHandle = HandleToAnchorMapPair.Key;
+		UARPin* const AnchorPin = HandleToAnchorMapPair.Value;
+		LuminArAnchor* LuminArAnchor = reinterpret_cast<struct LuminArAnchor*>(AnchorPin->GetNativeResource());
+		check(LuminArAnchor);
+		const MLHandle ParentTrackableHandle = LuminArAnchor->ParentTrackable;
+		if (ParentTrackableHandle != ML_INVALID_HANDLE)
+		{
+			UARTrackedGeometry* const ParentTrackable = Session->GetUObjectManager()->GetTrackableFromHandle<UARTrackedGeometry>(ParentTrackableHandle, Session);
+			const EARTrackingState AnchorTrackingState = ParentTrackable->GetTrackingState();
+			if (AnchorPin->GetTrackingState() != EARTrackingState::StoppedTracking)
+			{
+				AnchorPin->OnTrackingStateChanged(AnchorTrackingState);
+			}
+
+			if (AnchorPin->GetTrackingState() == EARTrackingState::Tracking)
+			{
+				AnchorPin->OnTransformUpdated(ParentTrackable->GetLocalToTrackingTransform());
+			}
+			UpdatedAnchors.Add(AnchorPin);
+		}
+	}
+#endif
+}
+
+//FTransform FLuminARFrame::GetCameraPose() const
+//{
+//	return LatestCameraPose;
+//}
+
+int64 FLuminARFrame::GetCameraTimestamp() const
+{
+	return LatestCameraTimestamp;
+}
+
+ELuminARTrackingState FLuminARFrame::GetCameraTrackingState() const
+{
+	return LatestCameraTrackingState;
+}
+
+void FLuminARFrame::GetUpdatedAnchors(TArray<UARPin*>& OutUpdatedAnchors) const
+{
+	OutUpdatedAnchors = UpdatedAnchors;
+}
+
+void FLuminARFrame::ARLineTrace(FVector2D ScreenPosition, ELuminARLineTraceChannel RequestedTraceChannels, TArray<FARTraceResult>& OutHitResults) const
+{
+#if PLATFORM_LUMIN
+	// Only testing straight forward from a little below the headset... Lumin isn't a handheld, but it's nice to have this do something.
+	IXRTrackingSystem* XRTrackingSystem = Session->GetARSystem()->GetXRTrackingSystem();
+	TArray<int32> Devices;
+	XRTrackingSystem->EnumerateTrackedDevices(Devices, EXRTrackedDeviceType::HeadMountedDisplay);
+	check(Devices.Num() == 1);
+	if (Devices.Num() > 0)
+	{
+		const int32 HMDDeviceID = Devices[0];
+		FQuat HMDQuat;
+		FVector HMDPosition;
+		const bool Success = XRTrackingSystem->GetCurrentPose(HMDDeviceID, HMDQuat, HMDPosition);
+		FTransform TrackingToWorldTransform = XRTrackingSystem->GetTrackingToWorldTransform();
+		if (Success)
+		{
+			const FVector HMDWorldPosition = TrackingToWorldTransform.TransformPosition(HMDPosition);
+			const FQuat HMDWorldQuat = TrackingToWorldTransform.TransformRotation(HMDQuat);
+			const FVector Start = HMDWorldPosition + FVector(0.0f, 0.0f, -10.0f);
+			const FVector Direction = HMDWorldQuat.Vector();
+			const FVector End = Start + (Direction * 10000.0f);
+
+			ARLineTrace(Start, End, RequestedTraceChannels, OutHitResults);
+		}
+	}
+#endif
+}
+void FLuminARFrame::ARLineTrace(FVector Start, FVector End, ELuminARLineTraceChannel RequestedTraceChannels, TArray<FARTraceResult>& OutHitResults) const
+{
+#if PLATFORM_LUMIN
+
+	// Only testing vs planes now, but not the ground plane.
+	ELuminARLineTraceChannel AllPlaneTraceChannels = /*ELuminARLineTraceChannel::InfinitePlane |*/ ELuminARLineTraceChannel::PlaneUsingExtent | ELuminARLineTraceChannel::PlaneUsingBoundaryPolygon;
+	if (!(RequestedTraceChannels & AllPlaneTraceChannels))
+		return;
+
+	TArray<UARPlaneGeometry*> Planes;
+	Session->GetAllTrackables(Planes);
+
+	for (UARPlaneGeometry* PPlane : Planes)
+	{
+		check(PPlane);
+		UARPlaneGeometry& Plane = *PPlane;
+
+		const FTransform LocalToWorld = Plane.GetLocalToWorldTransform();
+		const FVector PlaneOrigin = LocalToWorld.GetLocation();
+		const FVector PlaneNormal = LocalToWorld.TransformVectorNoScale(FVector(0, 0, 1));
+		const FVector Dir = End - Start;
+		// check if Dir is parallel to plane, no intersection
+		if (!FMath::IsNearlyZero(Dir | PlaneNormal, KINDA_SMALL_NUMBER))
+		{
+			// if T < 0 or > 1 we are outside the line segment, no intersection
+			float T = (((PlaneOrigin - Start) | PlaneNormal) / ((End - Start) | PlaneNormal));
+			if (T >= 0.0f || T <= 1.0f)
+			{
+				const FVector Intersection = Start + (Dir * T);
+
+				EARLineTraceChannels FoundInTraceChannel = EARLineTraceChannels::None;
+
+				if (!!(RequestedTraceChannels & (ELuminARLineTraceChannel::PlaneUsingExtent | ELuminARLineTraceChannel::PlaneUsingBoundaryPolygon)))
+				{
+					const FTransform WorldToLocal = LocalToWorld.Inverse();
+					const FVector LocalIntersection = WorldToLocal.TransformPosition(Intersection);
+
+					// Note: doing boundary check first for consistency with ARCore
+
+					if (FoundInTraceChannel == EARLineTraceChannels::None
+						&& !!(RequestedTraceChannels & ELuminARLineTraceChannel::PlaneUsingBoundaryPolygon))
+					{
+						// Note: could optimize this by computing an aligned boundary bounding rectangle during plane update and testing that first.
+
+						// Did we hit inside the boundary?
+						const TArray<FVector> Boundary = Plane.GetBoundaryPolygonInLocalSpace();
+						if (Boundary.Num() > 3) // has to be at least a triangle to have an inside
+						{
+							// This is the 'ray casting algorithm' for detecting if a point is inside a polygon.
+
+							// Draw a line from the point to the outside. Test for intersect with all edges.  If an odd number of edges are intersected the point is inside the polygon.
+							// The bounds are in plane local space, such that the plane is the x-y plane (z is always zero).
+							// We will offset that to put the point we are testing at 0,0 and our ray will be the +y axis (meaning the endpoint is 0,infinity and certainly outside the polygon).
+
+							// This could get the wrong answer if the test line goes exactly through a boundary vertex because that would register as two intersections.
+							// We are ignoring this rare failure cases.
+
+							const FVector2D Origin(LocalIntersection.X, LocalIntersection.Y);
+							const int Num = Boundary.Num();
+							FVector2D A(Boundary[Num - 1].X - Origin.X, Boundary[Num - 1].Y - Origin.Y);
+							int32 Crossings = 0;
+							for (int i = 0; i < Num; ++i)
+							{
+								const FVector2D B(Boundary[i].X - Origin.X, Boundary[i].Y - Origin.Y);
+
+								// Check if there is any Y intercept in the line segment.
+								if (FMath::Sign(A.X) != FMath::Sign(B.X))
+								{
+									// Check if the Y intercept is positive.
+									const float Slope = (B.Y - A.Y) / (B.X - A.X);
+									const float YIntercept = A.Y - (Slope * A.X);
+									if (YIntercept > 0.0f)
+									{
+										Crossings += 1;
+									}
+								}
+
+								A = B;
+							}
+							if ((Crossings & 0x01) == 0x01)
+							{
+								FoundInTraceChannel = EARLineTraceChannels::PlaneUsingBoundaryPolygon;
+							}
+						}
+					}
+
+					if (FoundInTraceChannel == EARLineTraceChannels::None
+						&& !!(RequestedTraceChannels & ELuminARLineTraceChannel::PlaneUsingExtent))
+					{
+						// Did we hit inside the plane extents?
+						if (FMath::Abs(LocalIntersection.X) <= Plane.GetExtent().X
+							&& FMath::Abs(LocalIntersection.Y) <= Plane.GetExtent().Y)
+						{
+							FoundInTraceChannel = EARLineTraceChannels::PlaneUsingExtent;
+						}
+					}
+				}
+
+				//// This 'infinite plane' 'ground plane' stuff seems... weird.
+				//if (FoundInTraceChannel == EARLineTraceChannels::None
+				//	&&!!(RequestedTraceChannels & ELuminARLineTraceChannel::InfinitePlane))
+				//{
+				//	FoundInTraceChannel = EARLineTraceChannels::GroundPlane;
+				//}
+
+				// write the result
+				if (FoundInTraceChannel != EARLineTraceChannels::None)
+				{
+					const float Distance = Dir.Size() * T;
+
+					FTransform HitTransform = LocalToWorld;
+					HitTransform.SetLocation(Intersection);
+
+					FARTraceResult UEHitResult(Session->GetARSystem(), Distance, FoundInTraceChannel, HitTransform, &Plane);
+					UEHitResult.SetLocalToWorldTransform(HitTransform);
+					OutHitResults.Add(UEHitResult);
+				}
+			}
+		}
+	}
+
+	// Sort closest to furthest
+	OutHitResults.Sort(FARTraceResult::FARTraceResultComparer());
+#endif
+}
+
+
+FMatrix FLuminARFrame::GetProjectionMatrix() const
+{
+	FMatrix ProjectionMatrix;
+
+#if PLATFORM_LUMIN
+	if (Session == nullptr)
+	{
+		return ProjectionMatrix;
+	}
+
+	//TODO - get projection from session
+
+	// Unreal uses the infinite far plane project matrix.
+	ProjectionMatrix.M[2][2] = 0.0f;
+	ProjectionMatrix.M[2][3] = 1.0f;
+	ProjectionMatrix.M[3][2] = GNearClippingPlane;
+#endif
+	return ProjectionMatrix;
+}
+
+void FLuminARFrame::TransformDisplayUvCoords(const TArray<float>& UvCoords, TArray<float>& OutUvCoords) const
+{
+#if PLATFORM_LUMIN
+	OutUvCoords = UvCoords;
+#endif
+}
+
+
+FLuminARLightEstimate FLuminARFrame::GetLightEstimate() const
+{
+#if PLATFORM_LUMIN
+	//TODO - Fill in light estimate.  See ml_lighting_tracking.h.
+	FLuminARLightEstimate LightEstimate;
+	return LightEstimate;
+#else
+	return FLuminARLightEstimate();
+#endif
+}
+
+void FLuminARFrame::StartPlaneQuery()
+{
+#if PLATFORM_LUMIN
+	//if we haven't queried yet, start one!
+	if (PlaneQueryHandle == ML_INVALID_HANDLE)
+	{
+		if (IMagicLeapPlugin::Get().IsMagicLeapHMDValid())
+		{
+			const FAppFramework& AppFramework = static_cast<FMagicLeapHMD*>(GEngine->XRSystem->GetHMDDevice())->GetAppFrameworkConst();
+			float WorldToMetersScale = AppFramework.GetWorldToMetersScale();
+			check(WorldToMetersScale != 0);
+
+			FTransform PoseInverse = UHeadMountedDisplayFunctionLibrary::GetTrackingToWorldTransform(nullptr).Inverse();
+			FPlane plane;
+
+			// Apply Lumin specific AR session config, if available.  Otherwise use default values.
+			MaxPlaneQueryResults = 200;
+			int32 MinPlaneArea = 25;
+			TArray<EPlaneQueryFlags> QueryFlags;
+			FVector SearchVolumeExtents(10000.0f, 10000.0f, 10000.0f);
+			bDiscardZeroExtentPlanes = false;
+
+			const UARSessionConfig& ARSessionConfig = Session->GetARSystem()->AccessSessionConfig();
+			const ULuminARSessionConfig* LuminARSessionConfig = Cast<ULuminARSessionConfig>(&ARSessionConfig);
+			if (LuminARSessionConfig != nullptr)
+			{
+				MaxPlaneQueryResults = LuminARSessionConfig->MaxPlaneQueryResults;
+				MinPlaneArea = LuminARSessionConfig->MinPlaneArea;
+				if (LuminARSessionConfig->ShouldDoHorizontalPlaneDetection())	{ QueryFlags.Add(EPlaneQueryFlags::Horizontal); }
+				if (LuminARSessionConfig->ShouldDoVerticalPlaneDetection())		{ QueryFlags.Add(EPlaneQueryFlags::Vertical); }
+				if (LuminARSessionConfig->bArbitraryOrientationPlaneDetection)	{ QueryFlags.Add(EPlaneQueryFlags::Arbitrary); }
+				SearchVolumeExtents = LuminARSessionConfig->PlaneSearchExtents;
+				for (EPlaneQueryFlags flag : LuminARSessionConfig->PlaneQueryFlags)
+				{
+					QueryFlags.Add(flag);
+				}
+				bDiscardZeroExtentPlanes = LuminARSessionConfig->bDiscardZeroExtentPlanes;
+			}
+			else
+			{
+				UE_LOG(LogLuminARAPI, Log, TEXT("LuminArSessionConfig not found, using defaults for lumin specific settings."));
+				QueryFlags.Add(EPlaneQueryFlags::Vertical);
+				QueryFlags.Add(EPlaneQueryFlags::Horizontal);
+			}
+
+			MLPlanesQuery Query;
+			Query.max_results = static_cast<uint32>(MaxPlaneQueryResults);
+			Query.flags = UnrealToMLPlanesQueryFlags(QueryFlags) | MLPlanesQueryFlag_Polygons;
+			Query.min_hole_length = 50.0f / WorldToMetersScale;  // ML docs say this value is deprecated, so presumably it does nothing now.
+			Query.min_plane_area = MinPlaneArea / (WorldToMetersScale * WorldToMetersScale);
+			Query.bounds_center = MagicLeap::ToMLVector(PoseInverse.GetTranslation(), WorldToMetersScale);
+			Query.bounds_rotation = MagicLeap::ToMLQuat(PoseInverse.GetRotation());
+			Query.bounds_extents = MagicLeap::ToMLVector(SearchVolumeExtents, WorldToMetersScale);
+
+			// MagicLeap::ToMLVector() causes the Z component to be negated.
+			// The bounds were thus invalid and resulted in everything being tracked. 
+			// This provides the content devs with an option to ignore the bounding volume at will.
+			{
+				Query.bounds_extents.x = FMath::Abs<float>(Query.bounds_extents.x);
+				Query.bounds_extents.y = FMath::Abs<float>(Query.bounds_extents.y);
+				Query.bounds_extents.z = FMath::Abs<float>(Query.bounds_extents.z);
+			}
+
+			MLResult QueryResult = MLPlanesQueryBegin(PlaneTrackerHandle, &Query, &PlaneQueryHandle);
+			if (QueryResult != MLResult_Ok || !MLHandleIsValid(PlaneQueryHandle))
+			{
+				UE_LOG(LogLuminARAPI, Error, TEXT("LuminARFrame could not request planes."));
+			}
+		}
+	}
+#endif //PLATFORM_LUMIN
+}
+
+void FLuminARFrame::ProcessPlaneQuery()
+{
+#if PLATFORM_LUMIN
+	if (PlaneQueryHandle != ML_INVALID_HANDLE)
+	{
+		FTransform PoseTransform = FTransform::Identity; 
+
+		uint32 OutNumResults = 0;
+		TArray<MLPlane> ResultMLPlanes;
+		ResultMLPlanes.AddDefaulted(MaxPlaneQueryResults);
+
+		MLPlaneBoundariesList PlaneBoundariesList;
+		MLPlaneBoundariesListInit(&PlaneBoundariesList);
+
+		MLResult PlaneQueryResult = MLPlanesQueryGetResultsWithBoundaries(PlaneTrackerHandle, PlaneQueryHandle, ResultMLPlanes.GetData(), &OutNumResults, &PlaneBoundariesList);
+		switch (PlaneQueryResult)
+		{
+		case MLResult_Pending:
+			// Intentionally skip. We'll continue to check until it has completed.
+			break;
+		case MLResult_UnspecifiedFailure:
+		{
+			UE_LOG(LogLuminARAPI, Error, TEXT("MLPlanesQueryGetResults MLResult_UnspecifiedFailure."));
+			PlaneQueryHandle = ML_INVALID_HANDLE;
+			LatestARPlaneQueryStatus = ELuminARPlaneQueryStatus::Fail;
+			break;
+		}
+		case MLResult_Ok:
+		{
+			const FAppFramework& AppFramework = static_cast<FMagicLeapHMD*>(GEngine->XRSystem->GetHMDDevice())->GetAppFrameworkConst();
+			float WorldToMetersScale = AppFramework.GetWorldToMetersScale();
+
+			PlaneResultsMap.Reset();
+			PlaneResultsMap.Reserve(OutNumResults);
+
+			// Setup for boundaries, build a map of Handles to Boundaries
+			TMap<uint64, const MLPlaneBoundaries*> HandleToBoundariesMap;
+			{
+				const uint32 BoundaryCount = PlaneBoundariesList.plane_boundaries_count;
+				for (uint32 b = 0; b < BoundaryCount; ++b)
+				{
+					const MLPlaneBoundaries& Boundaries = PlaneBoundariesList.plane_boundaries[b];
+					HandleToBoundariesMap.Add(Boundaries.id, &Boundaries);
+				}
+			}
+			const FRotator RotateToContentOrientation(-90.0f, 0.0f, 0.0f);
+			const FTransform RotateToContentTransform(RotateToContentOrientation);
+
+			for (uint32 i = 0; i < OutNumResults; ++i)
+			{
+				const MLPlane& ResultMLPlane = ResultMLPlanes[i];
+				FPlaneResult ResultUEPlane;
+
+				const uint64_t Mask = TNumericLimits<uint32>::Max();
+				ResultUEPlane.ID.A = ResultMLPlane.id & Mask;
+				ResultUEPlane.ID.B = ResultMLPlane.id >> 32;
+				ResultUEPlane.ID_64 = ResultMLPlane.id;
+
+				// Perception uses all coordinates in RUB so for them X axis is right and corresponds to the width of the plane.
+				// Unreal uses FRU, so the Y-axis is towards the right which makes the Y component of the vector the width.
+				ResultUEPlane.PlaneDimensions = FVector2D(ResultMLPlane.height * WorldToMetersScale, ResultMLPlane.width * WorldToMetersScale);
+				if (ResultUEPlane.PlaneDimensions.X == 0.0f || ResultUEPlane.PlaneDimensions.Y == 0.0f)
+				{
+					if (bDiscardZeroExtentPlanes)
+					{
+						continue;
+					}
+				}
+
+				FTransform planeTransform = FTransform(MagicLeap::ToFQuat(ResultMLPlane.rotation), MagicLeap::ToFVector(ResultMLPlane.position, WorldToMetersScale), FVector(1.0f, 1.0f, 1.0f));
+				if (planeTransform.ContainsNaN())
+				{
+					continue;
+				}
+				if (!planeTransform.GetRotation().IsNormalized())
+				{
+					FQuat rotation = planeTransform.GetRotation();
+					rotation.Normalize();
+					planeTransform.SetRotation(rotation);
+				}
+
+				planeTransform.SetRotation(MagicLeap::ToUERotator(planeTransform.GetRotation()));
+				planeTransform.AddToTranslation(PoseTransform.GetLocation());
+				planeTransform.ConcatenateRotation(PoseTransform.Rotator().Quaternion());
+				ResultUEPlane.PlanePosition = planeTransform.GetLocation();
+				ResultUEPlane.PlaneOrientation = planeTransform.Rotator();
+				// The plane orientation has the forward axis (X) pointing in the direction of the plane's normal.
+				// We are rotating it by 90 degrees clock-wise about the right axis (Y) to get the up vector (Z) to point in the direction of the plane's normal.
+				// ResultPlane we are rotating the axis, the rotation is in the opposite direction of the object i.e. -90 degrees.
+				ResultUEPlane.ContentOrientation = UKismetMathLibrary::Conv_VectorToRotator(UKismetMathLibrary::RotateAngleAxis(UKismetMathLibrary::Conv_RotatorToVector(ResultUEPlane.PlaneOrientation), -90, UKismetMathLibrary::GetRightVector(ResultUEPlane.PlaneOrientation)));
+				MLToUnrealPlanesQueryFlags(ResultMLPlane.flags, ResultUEPlane.PlaneFlags);
+
+				// Boundaries
+				const MLPlaneBoundaries** const BoundariesPtr = HandleToBoundariesMap.Find(ResultMLPlane.id);
+				const MLPlaneBoundaries* const Boundaries = BoundariesPtr ? *BoundariesPtr : nullptr;
+				if (Boundaries && Boundaries->boundaries_count > 0)
+				{
+					FTransform PlaneTransformInverse = planeTransform.Inverse();
+					FTransform BoundaryVertTransform = PlaneTransformInverse * RotateToContentTransform;
+
+					// Seems like there is really only one boundary polygon...
+					const int BoundariesCount = Boundaries->boundaries_count;
+					for (int b = 0; b < BoundariesCount; ++b)
+					{
+						const MLPlaneBoundary& Boundary = Boundaries->boundaries[b];
+						const MLPolygon* Polygon = Boundary.polygon;
+						check(Polygon);
+						const uint32 VertCount = Polygon->vertices_count;
+						for (uint32 v = 0; v < VertCount; ++v)
+						{
+							MLVec3f& Vert = Polygon->vertices[v];
+							FVector LocalVert = BoundaryVertTransform.TransformPosition(MagicLeap::ToFVector(Vert, WorldToMetersScale));
+							ResultUEPlane.BoundaryPolygon.Add(LocalVert);
+						}
+					}
+				}
+
+				PlaneResultsMap.Add(ResultMLPlane.id, ResultUEPlane);
+			}
+
+			PlaneQueryHandle = ML_INVALID_HANDLE;
+
+			// Mark planes that previously existed, but no longer do StoppedTracking.
+			const auto & TrackableHandleMap = Session->GetUObjectManager()->TrackableHandleMap;
+			for (auto TrackablePair : TrackableHandleMap)
+			{
+				uint64 Handle = TrackablePair.Key;
+				if (!PlaneResultsMap.Contains(Handle))
+				{
+					UARTrackedGeometry* ARTrackedGeometry = TrackablePair.Value.Get();
+					if (ARTrackedGeometry)
+					{
+						ARTrackedGeometry->SetTrackingState(EARTrackingState::StoppedTracking);
+					}
+				}
+			}
+
+			for (auto PlaneResultPair : PlaneResultsMap)
+			{
+				UARTrackedGeometry* ARTrackedGeometry = Session->GetUObjectManager()->GetTrackableFromHandle<UARPlaneGeometry>(PlaneResultPair.Key, Session);
+
+				if (ARTrackedGeometry && ARTrackedGeometry->GetTrackingState() != EARTrackingState::StoppedTracking)
+				{
+					ARTrackedGeometry->SetTrackingState(EARTrackingState::Tracking);
+					FLuminARTrackableResource* TrackableResource = reinterpret_cast<FLuminARTrackableResource*>(ARTrackedGeometry->GetNativeResource());
+					TrackableResource->UpdateGeometryData(Session);
+				}
+			}
+			LatestARPlaneQueryStatus = ELuminARPlaneQueryStatus::Success;
+			break;
+		}
+		default:
+			UE_LOG(LogLuminARAPI, Warning, TEXT("Unexpected return code from MLPlanesQueryGetResults: %d"), PlaneQueryResult);
+			LatestARPlaneQueryStatus = ELuminARPlaneQueryStatus::Fail;
+		}
+		MLPlanesReleaseBoundariesList(PlaneTrackerHandle, &PlaneBoundariesList);
+	}
+#endif //PLATFORM_LUMIN
+}
+
+
+
+TSharedPtr<FLuminARSession> FLuminARSession::CreateLuminARSession()
+{
+	TSharedPtr<FLuminARSession> NewSession = MakeShared<FLuminARSession>();
+
+	ULuminARUObjectManager* UObjectManager = NewObject<ULuminARUObjectManager>();
+	UObjectManager->AddToRoot();
+
+	NewSession->UObjectManager = UObjectManager;
+	return NewSession;
+}
+
+
+/************************************************/
+/*       ULuminARTrackableResource         */
+/************************************************/
+#if PLATFORM_LUMIN
+
+EARTrackingState FLuminARTrackableResource::GetTrackingState()
+{
+	if (MLHandleIsValid(TrackableHandle))
+	{
+		check(TrackedGeometry);
+		return TrackedGeometry->GetTrackingState();
+	}
+
+	return EARTrackingState::NotTracking;
+}
+
+void FLuminARTrackableResource::UpdateGeometryData(FLuminARSession* InSession)
+{
+	TrackedGeometry->UpdateTrackingState(GetTrackingState());
+}
+
+void FLuminARTrackableResource::ResetNativeHandle(LuminArTrackable* InTrackableHandle)
+{
+	TrackableHandle = ML_INVALID_HANDLE;
+	if (InTrackableHandle != nullptr)
+	{
+		TrackableHandle = InTrackableHandle->Handle;
+	}
+
+	UpdateGeometryData(nullptr);
+}
+
+void FLuminARTrackedPlaneResource::UpdateGeometryData(FLuminARSession* InSession)
+{
+	FLuminARTrackableResource::UpdateGeometryData(InSession);
+
+	UARPlaneGeometry* PlaneGeometry = CastChecked<UARPlaneGeometry>(TrackedGeometry);
+
+	if (!InSession || /*!CheckIsSessionValid("LuminARPlane", Session) || */TrackedGeometry->GetTrackingState() == EARTrackingState::StoppedTracking)
+	{
+		return;
+	}
+
+	const FLuminARFrame* Frame = InSession->GetLatestFrame();
+	if (Frame == nullptr)
+	{
+		return;
+	}
+
+	// PlaneResult is in unreal tracking space, so already scaled and axis corrected
+	const FPlaneResult* PlaneResult = Frame->GetPlaneResult(TrackableHandle);
+	check(PlaneResult);
+
+	FTransform LocalToTrackingTransform(PlaneResult->ContentOrientation, PlaneResult->PlanePosition);
+	FVector Extent(PlaneResult->PlaneDimensions.X * 0.5f, PlaneResult->PlaneDimensions.Y * 0.5f, 0);  // Extent is half the width and height
+
+	uint32 FrameNum = InSession->GetFrameNum();
+	int64 TimeStamp = InSession->GetLatestFrame()->GetCameraTimestamp();
+
+	UARPlaneGeometry* SubsumedByPlane = nullptr; // ARCore only
+	PlaneGeometry->UpdateTrackedGeometry(InSession->GetARSystem(), FrameNum, static_cast<double>(TimeStamp), LocalToTrackingTransform, InSession->GetARSystem()->GetAlignmentTransform(), FVector::ZeroVector, Extent, PlaneResult->BoundaryPolygon, SubsumedByPlane);
+	PlaneGeometry->SetDebugName(FName(TEXT("LuminARPlane")));
+}
+#endif
+
+
+#if PLATFORM_LUMIN
+void ULuminARUObjectManager::DumpTrackableHandleMap(const MLHandle SessionHandle)
+{
+	UE_LOG(LogLuminARAPI, Log, TEXT("ULuminARUObjectManager::DumpTrackableHandleMap"));
+	for (auto KeyValuePair : TrackableHandleMap)
+	{
+		uint64 TrackableHandle = KeyValuePair.Key;
+		TWeakObjectPtr<UARTrackedGeometry> TrackedGeometry = KeyValuePair.Value;
+		UE_LOG(LogLuminARAPI, Log, TEXT("  Trackable Handle %llu"), TrackableHandle);
+		if (TrackedGeometry.IsValid())
+		{
+			UARTrackedGeometry* TrackedGeometryObj = TrackedGeometry.Get();
+			FLuminARTrackableResource* NativeResource = reinterpret_cast<FLuminARTrackableResource*>(TrackedGeometryObj->GetNativeResource());
+			UE_LOG(LogLuminARAPI, Log, TEXT("  TrackedGeometry - NativeResource:%p, type: %s, tracking state: %d"),
+				NativeResource->GetNativeHandle(), *TrackedGeometryObj->GetClass()->GetFName().ToString(), (int)TrackedGeometryObj->GetTrackingState());
+		}
+		else
+		{
+			UE_LOG(LogLuminARAPI, Log, TEXT("  TrackedGeometry - InValid or Pending Kill."))
+		}
+	}
+}
+#endif
+
+ELuminARAPIStatus FLuminARSession::AcquireCameraImage(ULuminARCameraImage *&OutCameraImage)
+{
+	ELuminARAPIStatus ApiStatus = ELuminARAPIStatus::AR_SUCCESS;
+#if PLATFORM_LUMIN
+	if (LatestFrame == nullptr)
+	{
+		return ELuminARAPIStatus::AR_ERROR_FATAL;
+	}
+
+	ArImage *OutImage = nullptr;
+	//TODO - support acquiring camera
+
+	if (ApiStatus == ELuminARAPIStatus::AR_SUCCESS)
+	{
+		UE_LOG(LogLuminARAPI, Error, TEXT("Support camera image"));
+	}
+	else
+	{
+		UE_LOG(LogLuminARAPI, Error, TEXT("AcquireCameraImage failed!"));
+	}
+#endif
+
+	return ApiStatus;
+}
+
+void* FLuminARSession::GetLatestFrameRawPointer()
+{
+#if PLATFORM_LUMIN
+	//TODO - do we need this?
+#endif
+	return nullptr;
+}
+
+
+
+
diff --git a/Engine/Plugins/Lumin/MagicLeap/Source/MagicLeapAR/Private/LuminARAPI.h b/Engine/Plugins/Lumin/MagicLeap/Source/MagicLeapAR/Private/LuminARAPI.h
new file mode 100644
--- /dev/null
+++ b/Engine/Plugins/Lumin/MagicLeap/Source/MagicLeapAR/Private/LuminARAPI.h
@@ -0,0 +1,325 @@
+// Copyright 1998-2018 Epic Games, Inc. All Rights Reserved.
+#pragma once
+
+#include "CoreMinimal.h"
+#include "UObject/Object.h"
+
+#include "LuminARTypes.h"
+#include "LuminARSessionConfig.h"
+#include "ARSessionConfig.h"
+#include "PlanesComponent.h"
+
+#include "LuminARAPI.generated.h"
+
+DEFINE_LOG_CATEGORY_STATIC(LogLuminARAPI, Log, All);
+
+
+enum class ELuminARAPIStatus : int
+{
+	/// The operation was successful.
+	AR_SUCCESS = 0,
+
+	/// One of the arguments was invalid, either null or not appropriate for the
+	/// operation requested.
+	AR_ERROR_INVALID_ARGUMENT = -1,
+
+	/// An internal error occurred that the application should not attempt to
+	/// recover from.
+	AR_ERROR_FATAL = -2,
+};
+
+enum class ELuminARPlaneQueryStatus : int
+{
+	Unknown,
+	Success,
+	Fail
+};
+
+#if PLATFORM_LUMIN
+static ArTrackableType GetTrackableType(UClass* ClassType)
+{
+	if (ClassType == UARTrackedGeometry::StaticClass())
+	{
+		return ArTrackableType::LUMIN_AR_TRACKABLE_PLANE;
+	}
+	else if (ClassType == UARPlaneGeometry::StaticClass())
+	{
+		return ArTrackableType::LUMIN_AR_TRACKABLE_PLANE;
+	}
+	else
+	{
+		return ArTrackableType::LUMIN_AR_TRACKABLE_NOT_VALID;
+	}
+}
+#endif
+
+class FLuminARFrame;
+class FLuminARSession;
+class ULuminARCameraImage;
+
+UCLASS()
+class ULuminARUObjectManager : public UObject
+{
+	GENERATED_BODY()
+
+public:
+	UPROPERTY()
+	TArray<UARPin*> AllAnchors;
+
+#if PLATFORM_LUMIN
+	TMap<uint64, TSharedPtr<LuminArAnchor>> HandleToLuminAnchorMap;
+	TMap<uint64, UARPin*> HandleToAnchorMap;
+	TMap<uint64, TWeakObjectPtr<UARTrackedGeometry>> TrackableHandleMap;
+
+	template< class T > T* GetTrackableFromHandle(MLHandle TrackableHandle, FLuminARSession* Session);
+
+	void DumpTrackableHandleMap(const MLHandle SessionHandle);
+#endif
+};
+
+
+class FLuminARSession : public TSharedFromThis<FLuminARSession>, public FGCObject
+{
+
+public:
+	static TSharedPtr<FLuminARSession> CreateLuminARSession();
+
+	FLuminARSession();
+	~FLuminARSession();
+
+	// Properties
+	ULuminARUObjectManager* GetUObjectManager();
+	float GetWorldToMeterScale();
+	void SetARSystem(TSharedRef<FARSupportInterface , ESPMode::ThreadSafe> InArSystem) { ARSystem = InArSystem; }
+	TSharedRef<FARSupportInterface , ESPMode::ThreadSafe> GetARSystem() { return ARSystem.ToSharedRef(); }
+#if PLATFORM_LUMIN
+	MLHandle GetPlaneTrackerHandle();
+#endif
+
+	ELuminARAPIStatus Resume();
+	ELuminARAPIStatus Pause();
+	ELuminARAPIStatus Update(float WorldToMeterScale);
+	const FLuminARFrame* GetLatestFrame();
+	uint32 GetFrameNum();
+
+	// Anchor API
+	ELuminARAPIStatus CreateARAnchor(const FTransform& TransfromInTrackingSpace, UARTrackedGeometry* TrackedGeometry, USceneComponent* ComponentToPin, FName InDebugName, UARPin*& OutAnchor);
+	void DetachAnchor(UARPin* Anchor);
+
+	void GetAllAnchors(TArray<UARPin*>& OutAnchors) const;
+	template< class T > void GetAllTrackables(TArray<T*>& OutLuminARTrackableList);
+
+	ELuminARAPIStatus AcquireCameraImage(ULuminARCameraImage *&OutCameraImage);
+	
+	void* GetLatestFrameRawPointer();
+
+private:
+	void InitTracker();
+	void DestroyTracker();
+	virtual void AddReferencedObjects(FReferenceCollector& Collector) override;
+	FLuminARFrame* LatestFrame;
+	ULuminARUObjectManager* UObjectManager;
+	float CachedWorldToMeterScale;
+	uint32 FrameNumber;
+
+	TSharedPtr<FARSupportInterface , ESPMode::ThreadSafe> ARSystem;
+
+#if PLATFORM_LUMIN
+	MLHandle PlaneTrackerHandle = ML_INVALID_HANDLE;
+#endif
+};
+
+class FLuminARFrame
+{
+	friend class FLuminARSession;
+
+public:
+	FLuminARFrame(FLuminARSession* Session);
+	~FLuminARFrame();
+
+	void Init();
+
+	void Update(float WorldToMeterScale);
+
+	FTransform GetCameraPose() const;
+	int64 GetCameraTimestamp() const;
+	ELuminARTrackingState GetCameraTrackingState() const;
+
+	void GetUpdatedAnchors(TArray<UARPin*>& OutUpdatedAnchors) const;
+
+	const FPlaneResult* GetPlaneResult(int64 Handle) const { return PlaneResultsMap.Find(Handle); }
+
+	void ARLineTrace(FVector2D ScreenPosition, ELuminARLineTraceChannel RequestedTraceChannels, TArray<FARTraceResult>& OutHitResults) const;
+	void ARLineTrace(FVector Start, FVector End, ELuminARLineTraceChannel RequestedTraceChannels, TArray<FARTraceResult>& OutHitResults) const;
+
+	FMatrix GetProjectionMatrix() const;
+	void TransformDisplayUvCoords(const TArray<float>& UvCoords, TArray<float>& OutUvCoords) const;
+
+	FLuminARLightEstimate GetLightEstimate() const;
+
+private:
+
+	void StartPlaneQuery();
+	void ProcessPlaneQuery();
+
+
+	FLuminARSession* Session;
+	//FTransform LatestCameraPose;
+	int64 LatestCameraTimestamp;
+	ELuminARTrackingState LatestCameraTrackingState;
+	ELuminARPlaneQueryStatus LatestARPlaneQueryStatus;
+
+	TArray<UARPin*> UpdatedAnchors;
+	TMap<uint64, FPlaneResult> PlaneResultsMap;
+
+	int32 MaxPlaneQueryResults = 0;
+	bool bDiscardZeroExtentPlanes = false;
+
+#if PLATFORM_LUMIN
+	MLHandle PlaneTrackerHandle = ML_INVALID_HANDLE;
+	MLHandle PlaneQueryHandle = ML_INVALID_HANDLE;
+#endif
+};
+
+class FLuminARTrackableResource : public IARRef
+{
+public:
+	// IARRef interface
+	virtual void AddRef() override { }
+
+	virtual void RemoveRef() override
+	{
+#if PLATFORM_LUMIN
+		TrackableHandle = ML_INVALID_HANDLE;
+#endif
+	}
+
+#if PLATFORM_LUMIN
+public:
+	FLuminARTrackableResource(MLHandle InTrackableHandle, UARTrackedGeometry* InTrackedGeometry)
+		: TrackableHandle(InTrackableHandle)
+		, TrackedGeometry(InTrackedGeometry)
+	{
+		ensure(TrackableHandle != ML_INVALID_HANDLE);
+	}
+
+	virtual ~FLuminARTrackableResource()
+	{
+	}
+
+	EARTrackingState GetTrackingState();
+
+	virtual void UpdateGeometryData(FLuminARSession* InSession);
+
+	MLHandle GetNativeHandle() { return TrackableHandle; }
+
+	void ResetNativeHandle(LuminArTrackable* InTrackableHandle);
+
+protected:
+	MLHandle TrackableHandle;
+	UARTrackedGeometry* TrackedGeometry;
+#endif
+};
+
+class FLuminARTrackedPlaneResource : public FLuminARTrackableResource
+{
+public:
+#if PLATFORM_LUMIN
+	FLuminARTrackedPlaneResource(MLHandle InTrackableHandle, UARTrackedGeometry* InTrackedGeometry)
+		: FLuminARTrackableResource(InTrackableHandle, InTrackedGeometry)
+	{
+		ensure(TrackableHandle != ML_INVALID_HANDLE);
+	}
+
+	void UpdateGeometryData(FLuminARSession* InSession) override;
+
+	ArPlane* GetPlaneHandle() { return reinterpret_cast<ArPlane*>(TrackableHandle); }
+#endif
+};
+
+#if PLATFORM_LUMIN
+// Template function definition
+template< class T >
+T* ULuminARUObjectManager::GetTrackableFromHandle(MLHandle TrackableHandle, FLuminARSession* Session)
+{
+	if (!TrackableHandleMap.Contains(TrackableHandle)
+		|| !TrackableHandleMap[TrackableHandle].IsValid()
+		|| TrackableHandleMap[TrackableHandle]->GetTrackingState() == EARTrackingState::StoppedTracking)
+	{
+		// Add the trackable to the cache.
+		UARTrackedGeometry* NewTrackableObject = nullptr;
+		ArTrackableType TrackableType = ArTrackableType::LUMIN_AR_TRACKABLE_NOT_VALID;
+
+		check(Session);
+		const FLuminARFrame* Frame = Session->GetLatestFrame();
+		check(Frame);
+		const FPlaneResult* PlaneResult = Frame->GetPlaneResult(TrackableHandle);
+		if (PlaneResult)
+		{
+			TrackableType = ArTrackableType::LUMIN_AR_TRACKABLE_PLANE;
+		}
+
+		IARRef* NativeResource = nullptr;
+		if (TrackableType == ArTrackableType::LUMIN_AR_TRACKABLE_PLANE)
+		{
+			UARPlaneGeometry* PlaneObject = NewObject<UARPlaneGeometry>();
+			NewTrackableObject = static_cast<UARTrackedGeometry*>(PlaneObject);
+			NativeResource = new FLuminARTrackedPlaneResource(TrackableHandle, NewTrackableObject);
+		}
+		else
+		{
+			//checkf(false, TEXT("ULuminARUObjectManager failed to get a valid trackable %p. Unknow LuminAR Trackable Type."), TrackableHandle);
+			return nullptr;
+		}
+
+		// We should have a valid trackable object now.
+		checkf(NewTrackableObject, TEXT("Unknow LuminAR Trackable Type: %d"), TrackableType);
+
+		NewTrackableObject->InitializeNativeResource(NativeResource);
+		NativeResource = nullptr;
+
+		FLuminARTrackableResource* TrackableResource = reinterpret_cast<FLuminARTrackableResource*>(NewTrackableObject->GetNativeResource());
+
+		// no, we always do this in ProcessPlaneQuery
+		// Update the tracked geometry data using the native resource
+		//TrackableResource->UpdateGeometryData();
+		ensure(TrackableResource->GetTrackingState() != EARTrackingState::StoppedTracking);
+
+		TrackableHandleMap.Add(TrackableHandle, TWeakObjectPtr<UARTrackedGeometry>(NewTrackableObject));
+	}
+
+	T* Result = Cast<T>(TrackableHandleMap[TrackableHandle].Get());
+	//checkf(Result, TEXT("ULuminARUObjectManager failed to get a valid trackable %p from the map."), TrackableHandle);
+	return Result;
+}
+#endif
+
+template< class T >
+void FLuminARSession::GetAllTrackables(TArray<T*>& OutARCoreTrackableList)
+{
+	OutARCoreTrackableList.Empty();
+#if PLATFORM_LUMIN
+
+	ArTrackableType TrackableType = GetTrackableType(T::StaticClass());
+	if (TrackableType == ArTrackableType::LUMIN_AR_TRACKABLE_NOT_VALID)
+	{
+		UE_LOG(LogLuminARAPI, Error, TEXT("Invalid Trackable type: %s"), *T::StaticClass()->GetName());
+		return;
+	}
+
+	for (auto TrackableHandleMapPair : GetUObjectManager()->TrackableHandleMap)
+	{
+		TWeakObjectPtr<UARTrackedGeometry>& Value = TrackableHandleMapPair.Value;
+		if (Value.IsValid() &&
+			Value->GetTrackingState() != EARTrackingState::StoppedTracking
+			)
+		{
+			T* Trackable = Cast<T>(Value.Get());
+			if (Trackable)
+			{
+				OutARCoreTrackableList.Add(Trackable);
+			}
+		}	
+	}
+#endif
+}
diff --git a/Engine/Plugins/Lumin/MagicLeap/Source/MagicLeapAR/Private/LuminARDevice.cpp b/Engine/Plugins/Lumin/MagicLeap/Source/MagicLeapAR/Private/LuminARDevice.cpp
new file mode 100644
--- /dev/null
+++ b/Engine/Plugins/Lumin/MagicLeap/Source/MagicLeapAR/Private/LuminARDevice.cpp
@@ -0,0 +1,442 @@
+// Copyright 1998-2018 Epic Games, Inc. All Rights Reserved.
+
+#include "LuminARDevice.h"
+#include "CoreMinimal.h"
+#include "Misc/ScopeLock.h"
+#include "HAL/ThreadSafeCounter64.h"
+#include "GameFramework/WorldSettings.h"
+#include "Framework/Application/SlateApplication.h"
+#include "Engine/World.h" // for FWorldDelegates
+#include "Engine/Engine.h"
+#include "GeneralProjectSettings.h"
+#include "LuminARTrackingSystem.h"
+
+namespace
+{
+	ELuminARFunctionStatus ToLuminARFunctionStatus(ELuminARAPIStatus Status)
+	{
+		switch (Status)
+		{
+		case ELuminARAPIStatus::AR_SUCCESS:
+			return ELuminARFunctionStatus::Success;
+		default:
+			ensureMsgf(false, TEXT("Unknown conversion from ELuminARAPIStatus %d to ELuminARFunctionStatus."), static_cast<int>(Status));
+			return ELuminARFunctionStatus::Fatal;
+		}
+	}
+}
+
+FLuminARDevice* FLuminARDevice::GetInstance()
+{
+	static FLuminARDevice Inst;
+	return &Inst;
+}
+
+FLuminARDevice::FLuminARDevice()
+	: bIsLuminARSessionRunning(false)
+	, bForceLateUpdateEnabled(false)
+	, bSessionConfigChanged(false)
+	, bStartSessionRequested(false)
+	, bShouldSessionRestart(false)
+	, WorldToMeterScale(100.0f)
+	, CurrentSessionStatus(EARSessionStatus::NotStarted)
+{
+}
+
+
+bool FLuminARDevice::GetIsTrackingTypeSupported(EARSessionType SessionType)
+{
+	if (SessionType == EARSessionType::World)
+	{
+		return true;
+	}
+	else
+	{
+		return false;
+	}
+}
+
+void FLuminARDevice::OnModuleLoaded()
+{
+	FWorldDelegates::OnWorldTickStart.AddRaw(this, &FLuminARDevice::OnWorldTickStart);
+}
+
+void FLuminARDevice::OnModuleUnloaded()
+{
+	FWorldDelegates::OnWorldTickStart.RemoveAll(this);
+	// clear the unique ptr.
+	LuminARSession.Reset();
+}
+
+bool FLuminARDevice::GetIsLuminARSessionRunning()
+{
+	return bIsLuminARSessionRunning;
+}
+
+EARSessionStatus FLuminARDevice::GetSessionStatus()
+{
+	return CurrentSessionStatus;
+}
+
+float FLuminARDevice::GetWorldToMetersScale()
+{
+	return WorldToMeterScale;
+}
+
+// This function will be called by public function to start Lumin AR session request.
+void FLuminARDevice::StartLuminARSessionRequest(UARSessionConfig* SessionConfig)
+{
+	UE_LOG(LogLuminAR, Log, TEXT("Start LuminAR session requested"));
+
+	if (bIsLuminARSessionRunning)
+	{
+		if (SessionConfig == AccessSessionConfig())
+		{
+			UE_LOG(LogLuminAR, Warning, TEXT("LuminAR session is already running with the requested LuminAR config. Request aborted."));
+			bStartSessionRequested = false;
+			return;
+		}
+
+		PauseLuminARSession();
+	}
+
+	if (bStartSessionRequested)
+	{
+		UE_LOG(LogLuminAR, Warning, TEXT("LuminAR session is already starting. This will overriding the previous session config with the new one."))
+	}
+
+	bStartSessionRequested = true;
+
+	// Try recreating the LuminARSession to fix the fatal error.
+	if (CurrentSessionStatus == EARSessionStatus::FatalError)
+	{
+		UE_LOG(LogLuminAR, Warning, TEXT("Reset LuminAR session due to fatal error detected."));
+		ResetLuminARSession();
+	}
+}
+
+bool FLuminARDevice::GetStartSessionRequestFinished()
+{
+	return !bStartSessionRequested;
+}
+
+// Note that this function will only be registered when LuminAR is supported.
+void FLuminARDevice::OnWorldTickStart(ELevelTick TickType, float DeltaTime)
+{
+	WorldToMeterScale = GWorld->GetWorldSettings()->WorldToMeters;
+	TFunction<void()> Func;
+	while (RunOnGameThreadQueue.Dequeue(Func))
+	{
+		Func();
+	}
+
+	if (!bIsLuminARSessionRunning && bStartSessionRequested)
+	{
+		StartSessionWithRequestedConfig();
+	}
+
+	if (bIsLuminARSessionRunning)
+	{
+		ELuminARAPIStatus Status = LuminARSession->Update(WorldToMeterScale);
+		if (Status == ELuminARAPIStatus::AR_ERROR_FATAL)
+		{
+			LuminARSession->Pause();
+			bIsLuminARSessionRunning = false;
+			CurrentSessionStatus = EARSessionStatus::FatalError;
+		}
+		else
+		{
+		}
+	}
+}
+
+void FLuminARDevice::HandleRuntimePermissionsGranted(const TArray<FString>& RuntimePermissions, const TArray<bool>& Granted)
+{
+	bool bGranted = true;
+	for (int32 i = 0; i < RuntimePermissions.Num(); i++)
+	{
+		if (!Granted[i])
+		{
+			bGranted = false;
+			UE_LOG(LogLuminAR, Warning, TEXT("Android runtime permission denied: %s"), *RuntimePermissions[i]);
+		}
+		else
+		{
+			UE_LOG(LogLuminAR, Log, TEXT("Android runtime permission granted: %s"), *RuntimePermissions[i]);
+		}
+	}
+}
+
+void FLuminARDevice::StartSessionWithRequestedConfig()
+{
+	bStartSessionRequested = false;
+
+	if (!LuminARSession.IsValid())
+	{
+		LuminARSession = FLuminARSession::CreateLuminARSession();
+		/*ELuminARAPIStatus SessionCreateStatus = LuminARSession->GetSessionCreateStatus();
+		if (SessionCreateStatus != ELuminARAPIStatus::AR_SUCCESS)
+		{
+			ensureMsgf(false, TEXT("Failed to create LuminAR session with error status: %d"), (int)SessionCreateStatus);
+			if (SessionCreateStatus != ELuminARAPIStatus::AR_ERROR_FATAL)
+			{
+				CurrentSessionStatus = EARSessionStatus::NotSupported;
+			}
+			else
+			{
+				CurrentSessionStatus = EARSessionStatus::FatalError;
+			}
+			LuminARSession.Reset();
+			return;
+		}
+		*/
+		LuminARSession->SetARSystem(ARSystem.ToSharedRef());
+	}
+
+	StartSession();
+}
+
+void FLuminARDevice::StartSession()
+{
+	UARSessionConfig* RequestedConfig = AccessSessionConfig();
+
+	if (RequestedConfig->GetSessionType() != EARSessionType::World)
+	{
+		UE_LOG(LogLuminAR, Warning, TEXT("Start AR failed: Unsupported AR tracking type %d for LuminAR"), static_cast<int>(RequestedConfig->GetSessionType()));
+		CurrentSessionStatus = EARSessionStatus::UnsupportedConfiguration;
+		return;
+	}
+
+	ELuminARAPIStatus Status = ELuminARAPIStatus::AR_SUCCESS;
+	/*
+	ELuminARAPIStatus Status = LuminARSession->ConfigSession(*RequestedConfig);
+
+	if (Status != ELuminARAPIStatus::AR_SUCCESS)
+	{
+	UE_LOG(LogLuminAR, Error, TEXT("LuminAR Session start failed with error status %d"), static_cast<int>(Status));
+	CurrentSessionStatus = EARSessionStatus::UnsupportedConfiguration;
+	return;
+	}
+	*/
+
+	Status = LuminARSession->Resume();
+
+	if (Status != ELuminARAPIStatus::AR_SUCCESS)
+	{
+		UE_LOG(LogLuminAR, Error, TEXT("LuminAR Session start failed with error status %d"), static_cast<int>(Status));
+		check(Status == ELuminARAPIStatus::AR_ERROR_FATAL);
+		// If we failed here, the only reason would be fatal error.
+		CurrentSessionStatus = EARSessionStatus::FatalError;
+		return;
+	}
+	/*
+	if (GEngine->XRSystem.IsValid())
+	{
+		FMagicLeapHMD* MagicLeapHMD = static_cast<FMagicLeapHMD*>(GEngine->XRSystem.Get());
+		if (MagicLeapHMD && LuminARImplementation.IsValid())
+		{
+			const bool bMatchFOV = RequestedConfig->ShouldRenderCameraOverlay();
+			LuminARImplementation->ConfigLuminARCamera(bMatchFOV, RequestedConfig->ShouldRenderCameraOverlay());
+		}
+		else
+		{
+			UE_LOG(LogLuminAR, Error, TEXT("ERROR: LuminARImplementation is not available."));
+		}
+	}
+	*/
+
+	bIsLuminARSessionRunning = true;
+	CurrentSessionStatus = EARSessionStatus::Running;
+
+	ARSystem->OnARSessionStarted.Broadcast();
+}
+
+void* FLuminARDevice::GetARSessionRawPointer()
+{
+#if PLATFORM_LUMIN
+	check(false); // never used and unsupported
+	return nullptr;
+#endif
+	return nullptr;
+}
+
+void* FLuminARDevice::GetGameThreadARFrameRawPointer()
+{
+#if PLATFORM_LUMIN
+	return reinterpret_cast<void*>(LuminARSession->GetLatestFrameRawPointer());
+#endif
+	return nullptr;
+}
+
+TSharedPtr<FARSupportInterface , ESPMode::ThreadSafe> FLuminARDevice::GetARSystem()
+{
+	return ARSystem;
+}
+
+
+void FLuminARDevice::SetARSystem(TSharedPtr<FARSupportInterface , ESPMode::ThreadSafe> InARSystem)
+{
+	check(InARSystem.IsValid());
+	ARSystem = InARSystem;
+}
+
+TSharedPtr<FLuminARImplementation, ESPMode::ThreadSafe> FLuminARDevice::GetLuminARImplementation()
+{
+	return LuminARImplementation;
+}
+
+void FLuminARDevice::SetLuminARImplementation(TSharedPtr<FLuminARImplementation, ESPMode::ThreadSafe> InARImplementation)
+{
+	LuminARImplementation = InARImplementation;
+}
+
+void FLuminARDevice::PauseLuminARSession()
+{
+	UE_LOG(LogLuminAR, Log, TEXT("Pausing LuminAR session."));
+	if (!bIsLuminARSessionRunning)
+	{
+		if(bStartSessionRequested)
+		{
+			bStartSessionRequested = false;
+		}
+		else
+		{
+			UE_LOG(LogLuminAR, Log, TEXT("Could not stop LuminAR tracking session because there is no running tracking session!"));
+		}
+		return;
+	}
+
+	ELuminARAPIStatus Status = LuminARSession->Pause();
+
+	if (Status == ELuminARAPIStatus::AR_ERROR_FATAL)
+	{
+		CurrentSessionStatus = EARSessionStatus::FatalError;
+	}
+
+	bIsLuminARSessionRunning = false;
+	CurrentSessionStatus = EARSessionStatus::NotStarted;
+}
+
+void FLuminARDevice::ResetLuminARSession()
+{
+	LuminARSession.Reset();
+	CurrentSessionStatus = EARSessionStatus::NotStarted;
+}
+
+
+FMatrix FLuminARDevice::GetPassthroughCameraProjectionMatrix(FIntPoint ViewRectSize) const
+{
+	if (!bIsLuminARSessionRunning)
+	{
+		return FMatrix::Identity;
+	}
+	return LuminARSession->GetLatestFrame()->GetProjectionMatrix();
+}
+
+void FLuminARDevice::GetPassthroughCameraImageUVs(const TArray<float>& InUvs, TArray<float>& OutUVs) const
+{
+	if (!bIsLuminARSessionRunning)
+	{
+		return;
+	}
+	LuminARSession->GetLatestFrame()->TransformDisplayUvCoords(InUvs, OutUVs);
+}
+
+ELuminARTrackingState FLuminARDevice::GetTrackingState() const
+{
+	if (!bIsLuminARSessionRunning)
+	{
+		return ELuminARTrackingState::StoppedTracking;
+	}
+	return LuminARSession->GetLatestFrame()->GetCameraTrackingState();
+}
+
+//FTransform FLuminARDevice::GetLatestPose() const
+//{
+//	if (!bIsLuminARSessionRunning)
+//	{
+//		return FTransform::Identity;
+//	}
+//	return LuminARSession->GetLatestFrame()->GetCameraPose();
+//}
+
+FLuminARLightEstimate FLuminARDevice::GetLatestLightEstimate() const
+{
+	if (!bIsLuminARSessionRunning)
+	{
+		return FLuminARLightEstimate();
+	}
+
+	return LuminARSession->GetLatestFrame()->GetLightEstimate();
+}
+
+void FLuminARDevice::ARLineTrace(const FVector2D& ScreenPosition, ELuminARLineTraceChannel TraceChannels, TArray<FARTraceResult>& OutHitResults)
+{
+	if (!bIsLuminARSessionRunning)
+	{
+		return;
+	}
+	OutHitResults.Empty();
+	LuminARSession->GetLatestFrame()->ARLineTrace(ScreenPosition, TraceChannels, OutHitResults);
+}
+
+void FLuminARDevice::ARLineTrace(const FVector& Start, const FVector& End, ELuminARLineTraceChannel TraceChannels, TArray<FARTraceResult>& OutHitResults)
+{
+	if (!bIsLuminARSessionRunning)
+	{
+		return;
+	}
+	OutHitResults.Empty();
+	LuminARSession->GetLatestFrame()->ARLineTrace(Start, End, TraceChannels, OutHitResults);
+}
+
+ELuminARFunctionStatus FLuminARDevice::CreateARPin(const FTransform& PinToWorldTransform, UARTrackedGeometry* TrackedGeometry, USceneComponent* ComponentToPin, const FName DebugName, UARPin*& OutARAnchorObject)
+{
+	if (!bIsLuminARSessionRunning)
+	{
+		return ELuminARFunctionStatus::SessionPaused;
+	}
+
+	const FTransform& TrackingToAlignedTracking = ARSystem->GetAlignmentTransform();
+	const FTransform PinToTrackingTransform = PinToWorldTransform.GetRelativeTransform(ARSystem->GetXRTrackingSystem()->GetTrackingToWorldTransform()).GetRelativeTransform(TrackingToAlignedTracking);
+
+	ELuminARFunctionStatus Status = ToLuminARFunctionStatus(LuminARSession->CreateARAnchor(PinToTrackingTransform, TrackedGeometry, ComponentToPin, DebugName, OutARAnchorObject));
+
+	return Status;
+}
+
+void FLuminARDevice::RemoveARPin(UARPin* ARAnchorObject)
+{
+	if (!bIsLuminARSessionRunning)
+	{
+		return;
+	}
+
+	LuminARSession->DetachAnchor(ARAnchorObject);
+}
+
+void FLuminARDevice::GetAllARPins(TArray<UARPin*>& LuminARAnchorList)
+{
+	if (!bIsLuminARSessionRunning)
+	{
+		return;
+	}
+	LuminARSession->GetAllAnchors(LuminARAnchorList);
+}
+
+UARSessionConfig* FLuminARDevice::AccessSessionConfig() const
+{
+	return (ARSystem.IsValid())
+		? &ARSystem->AccessSessionConfig()
+		: nullptr;
+}
+
+ELuminARFunctionStatus FLuminARDevice::AcquireCameraImage(ULuminARCameraImage *&OutLatestCameraImage)
+{
+	if (!bIsLuminARSessionRunning)
+	{
+		return ELuminARFunctionStatus::SessionPaused;
+	}
+
+	return ToLuminARFunctionStatus(LuminARSession->AcquireCameraImage(OutLatestCameraImage));
+}
diff --git a/Engine/Plugins/Lumin/MagicLeap/Source/MagicLeapAR/Private/LuminARDevice.h b/Engine/Plugins/Lumin/MagicLeap/Source/MagicLeapAR/Private/LuminARDevice.h
new file mode 100644
--- /dev/null
+++ b/Engine/Plugins/Lumin/MagicLeap/Source/MagicLeapAR/Private/LuminARDevice.h
@@ -0,0 +1,125 @@
+// Copyright 1998-2018 Epic Games, Inc. All Rights Reserved.
+
+#pragma once
+
+#include "Misc/CoreDelegates.h"
+#include "Engine/EngineBaseTypes.h"
+#include "HAL/ThreadSafeBool.h"
+#include "Containers/Queue.h"
+#include "Containers/Map.h"
+
+#include "LuminARTypes.h"
+#include "LuminARAPI.h"
+
+class FLuminARImplementation;
+
+class FLuminARDevice
+{
+public:
+	static FLuminARDevice* GetInstance();
+
+	FLuminARDevice();
+
+	bool GetIsTrackingTypeSupported(EARSessionType SessionType);
+
+	bool GetIsLuminARSessionRunning();
+
+	EARSessionStatus GetSessionStatus();
+
+	// Get Unreal Units per meter, based off of the current map's VR World to Meters setting.
+	float GetWorldToMetersScale();
+
+	// Start ARSession with custom session config.
+	void StartLuminARSessionRequest(UARSessionConfig* SessionConfig);
+
+	bool GetStartSessionRequestFinished();
+
+	void PauseLuminARSession();
+
+	void ResetLuminARSession();
+
+	// Passthrough Camera
+	FMatrix GetPassthroughCameraProjectionMatrix(FIntPoint ViewRectSize) const;
+	void GetPassthroughCameraImageUVs(const TArray<float>& InUvs, TArray<float>& OutUVs) const;
+
+	// Frame
+	ELuminARTrackingState GetTrackingState() const;
+	//FTransform GetLatestPose() const;
+	FLuminARLightEstimate GetLatestLightEstimate() const;
+	// Hit test
+	void ARLineTrace(const FVector2D& ScreenPosition, ELuminARLineTraceChannel TraceChannels, TArray<FARTraceResult>& OutHitResults);
+	void ARLineTrace(const FVector& Start, const FVector& End, ELuminARLineTraceChannel TraceChannels, TArray<FARTraceResult>& OutHitResults);
+
+	// Anchor, Planes
+	ELuminARFunctionStatus CreateARPin(const FTransform& PinToWorldTransform, UARTrackedGeometry* TrackedGeometry, USceneComponent* ComponentToPin, const FName DebugName, UARPin*& OutARAnchorObject);
+	void RemoveARPin(UARPin* ARAnchorObject);
+
+	void GetAllARPins(TArray<UARPin*>& LuminARAnchorList);
+
+	template< class T >
+	void GetAllTrackables(TArray<T*>& OutLuminARTrackableList)
+	{
+		if (!bIsLuminARSessionRunning)
+		{
+			return;
+		}
+		LuminARSession->GetAllTrackables<T>(OutLuminARTrackableList);
+	}
+
+	void RunOnGameThread(TFunction<void()> Func)
+	{
+		RunOnGameThreadQueue.Enqueue(Func);
+	}
+
+	void GetRequiredRuntimePermissionsForConfiguration(const UARSessionConfig& Config, TArray<FString>& RuntimePermissions)
+	{
+		RuntimePermissions.Reset();
+		// TODO: check for depth camera when it is supported here.
+		RuntimePermissions.Add("android.permission.CAMERA");
+	}
+	void HandleRuntimePermissionsGranted(const TArray<FString>& Permissions, const TArray<bool>& Granted);
+
+	// Function that is used to call from the Android UI thread:
+	void StartSessionWithRequestedConfig();
+
+	TSharedPtr<FARSupportInterface , ESPMode::ThreadSafe> GetARSystem();
+	void SetARSystem(TSharedPtr<FARSupportInterface , ESPMode::ThreadSafe> InARSystem);
+
+	TSharedPtr<FLuminARImplementation, ESPMode::ThreadSafe> GetLuminARImplementation();
+	void SetLuminARImplementation(TSharedPtr<FLuminARImplementation, ESPMode::ThreadSafe> InARImplementation);
+
+	ELuminARFunctionStatus AcquireCameraImage(ULuminARCameraImage *&OutLatestCameraImage);
+
+	void* GetARSessionRawPointer();
+	void* GetGameThreadARFrameRawPointer();
+
+private:
+	// Unreal plugin events.
+	void OnModuleLoaded();
+	void OnModuleUnloaded();
+
+	void OnWorldTickStart(ELevelTick TickType, float DeltaTime);
+
+	void StartSession();
+
+	friend class FLuminARAndroidHelper;
+	friend class FLuminARModule;
+
+	UARSessionConfig* AccessSessionConfig() const;
+
+private:
+	TSharedPtr<FLuminARSession> LuminARSession;
+	bool bIsLuminARSessionRunning;
+	bool bForceLateUpdateEnabled; // A debug flag to force use late update.
+	bool bSessionConfigChanged;
+	bool bStartSessionRequested; // User called StartSession
+	bool bShouldSessionRestart; // Start tracking on activity start
+	float WorldToMeterScale;
+
+	EARSessionStatus CurrentSessionStatus;
+
+	TQueue<TFunction<void()>> RunOnGameThreadQueue;
+
+	TSharedPtr<FARSupportInterface , ESPMode::ThreadSafe> ARSystem;
+	TSharedPtr<FLuminARImplementation, ESPMode::ThreadSafe> LuminARImplementation;
+};
diff --git a/Engine/Plugins/Lumin/MagicLeap/Source/MagicLeapAR/Private/LuminARFunctionLibrary.cpp b/Engine/Plugins/Lumin/MagicLeap/Source/MagicLeapAR/Private/LuminARFunctionLibrary.cpp
new file mode 100644
--- /dev/null
+++ b/Engine/Plugins/Lumin/MagicLeap/Source/MagicLeapAR/Private/LuminARFunctionLibrary.cpp
@@ -0,0 +1,84 @@
+// Copyright 1998-2018 Epic Games, Inc. All Rights Reserved.
+
+#include "LuminARFunctionLibrary.h"
+#include "UnrealEngine.h"
+#include "Engine/Engine.h"
+#include "LatentActions.h"
+#include "ARBlueprintLibrary.h"
+
+#include "LuminARDevice.h"
+#include "LuminARTrackingSystem.h"
+
+
+struct FLuminARStartSessionAction : public FPendingLatentAction
+{
+public:
+	FName ExecutionFunction;
+	int32 OutputLink;
+	FWeakObjectPtr CallbackTarget;
+
+	FLuminARStartSessionAction(const FLatentActionInfo& InLatentInfo)
+		: FPendingLatentAction()
+		, ExecutionFunction(InLatentInfo.ExecutionFunction)
+		, OutputLink(InLatentInfo.Linkage)
+		, CallbackTarget(InLatentInfo.CallbackTarget)
+	{}
+
+	virtual void UpdateOperation(FLatentResponse& Response) override
+	{
+		bool bSessionStartedFinished = FLuminARDevice::GetInstance()->GetStartSessionRequestFinished();
+		Response.FinishAndTriggerIf(bSessionStartedFinished, ExecutionFunction, OutputLink, CallbackTarget);
+	}
+#if WITH_EDITOR
+	virtual FString GetDescription() const override
+	{
+		return FString::Printf(TEXT("Starting LuminAR tracking session"));
+	}
+#endif
+};
+
+void ULuminARSessionFunctionLibrary::StartLuminARSession(UObject* WorldContextObject, struct FLatentActionInfo LatentInfo, ULuminARSessionConfig* Configuration)
+{
+	UE_LOG(LogTemp, Log, TEXT("ULuminARSessionFunctionLibrary::StartLuminARSession"));
+	if (UWorld* World = GEngine->GetWorldFromContextObject(WorldContextObject, EGetWorldErrorMode::LogAndReturnNull))
+	{
+		FLatentActionManager& LatentManager = World->GetLatentActionManager();
+		if (LatentManager.FindExistingAction<FLuminARStartSessionAction>(LatentInfo.CallbackTarget, LatentInfo.UUID) == nullptr)
+		{
+			UARBlueprintLibrary::StartARSession(static_cast<UARSessionConfig*>(Configuration));
+			FLuminARStartSessionAction* NewAction = new FLuminARStartSessionAction(LatentInfo);
+			LatentManager.AddNewAction(LatentInfo.CallbackTarget, LatentInfo.UUID, NewAction);
+		}
+	}
+}
+
+/************************************************************************/
+/*  ULuminARFrameFunctionLibrary                                   */
+/************************************************************************/
+ELuminARTrackingState ULuminARFrameFunctionLibrary::GetTrackingState()
+{
+	return FLuminARDevice::GetInstance()->GetTrackingState();
+}
+
+//void ULuminARFrameFunctionLibrary::GetPose(FTransform& LastePose)
+//{
+//	LastePose = FLuminARDevice::GetInstance()->GetLatestPose();
+//}
+
+bool ULuminARFrameFunctionLibrary::LuminARLineTrace(UObject* WorldContextObject, const FVector2D& ScreenPosition, TSet<ELuminARLineTraceChannel> TraceChannels, TArray<FARTraceResult>& OutHitResults)
+{
+	ELuminARLineTraceChannel TraceChannelValue = ELuminARLineTraceChannel::None;
+	for (ELuminARLineTraceChannel Channel : TraceChannels)
+	{
+		TraceChannelValue = TraceChannelValue | Channel;
+	}
+
+	FLuminARDevice::GetInstance()->ARLineTrace(ScreenPosition, TraceChannelValue, OutHitResults);
+	return OutHitResults.Num() > 0;
+}
+
+void ULuminARFrameFunctionLibrary::GetLightEstimation(FLuminARLightEstimate& LightEstimation)
+{
+	LightEstimation = FLuminARDevice::GetInstance()->GetLatestLightEstimate();
+}
+
diff --git a/Engine/Plugins/Lumin/MagicLeap/Source/MagicLeapAR/Private/LuminARModule.cpp b/Engine/Plugins/Lumin/MagicLeap/Source/MagicLeapAR/Private/LuminARModule.cpp
new file mode 100644
--- /dev/null
+++ b/Engine/Plugins/Lumin/MagicLeap/Source/MagicLeapAR/Private/LuminARModule.cpp
@@ -0,0 +1,73 @@
+// Copyright 1998-2018 Epic Games, Inc. All Rights Reserved.
+
+#include "LuminARModule.h"
+#include "Modules/ModuleManager.h"
+#include "Features/IModularFeatures.h"
+#include "Features/IModularFeature.h"
+
+#include "LuminARTrackingSystem.h"
+#include "LuminARDevice.h"
+#include "Templates/SharedPointer.h"
+
+
+#define LOCTEXT_NAMESPACE "LuminAR"
+
+class FLuminARModule : public ILuminARModule
+{
+	/** IModuleInterface implementation */
+	virtual void StartupModule() override;
+	virtual void ShutdownModule() override;
+
+	//create for mutual connection (regardless of construction order)
+	virtual TSharedPtr<IARSystemSupport, ESPMode::ThreadSafe> CreateARImplementation() override;
+	//Now connect (regardless of connection order)
+	virtual void ConnectARImplementationToXRSystem(FXRTrackingSystemBase* InXRTrackingSystem) override;
+	//Now initialize fully connected systems
+	virtual void InitializeARImplementation() override;
+
+private:
+	TSharedPtr<FLuminARImplementation, ESPMode::ThreadSafe> LuminARImplementation;
+};
+
+IMPLEMENT_MODULE(FLuminARModule, MagicLeapAR)
+
+
+void FLuminARModule::StartupModule()
+{
+	ensureMsgf(FModuleManager::Get().LoadModule("AugmentedReality"), TEXT("Lumin AR depends on the AugmentedReality module.") );
+
+	FModuleManager::LoadModulePtr<IModuleInterface>("AugmentedReality");
+}
+
+void FLuminARModule::ShutdownModule()
+{
+	// Complete LuminAR teardown.
+	FLuminARDevice::GetInstance()->OnModuleUnloaded();
+}
+
+
+//create for mutual connection (regardless of construction order)
+TSharedPtr<IARSystemSupport, ESPMode::ThreadSafe> FLuminARModule::CreateARImplementation()
+{
+#if PLATFORM_LUMIN
+	LuminARImplementation = MakeShareable(new FLuminARImplementation());
+#endif
+	return LuminARImplementation;
+}
+
+void FLuminARModule::ConnectARImplementationToXRSystem(FXRTrackingSystemBase* InXRTrackingSystem)
+{
+	ensure(InXRTrackingSystem);
+
+	FLuminARDevice::GetInstance()->SetLuminARImplementation(LuminARImplementation);
+	FLuminARDevice::GetInstance()->SetARSystem(InXRTrackingSystem->GetARCompositionComponent());
+	InXRTrackingSystem->GetARCompositionComponent()->InitializeARSystem();
+}
+
+void FLuminARModule::InitializeARImplementation()
+{
+	// Complete LuminAR setup.
+	FLuminARDevice::GetInstance()->OnModuleLoaded();
+}
+
+#undef LOCTEXT_NAMESPACE
diff --git a/Engine/Plugins/Lumin/MagicLeap/Source/MagicLeapAR/Private/LuminARTrackingSystem.cpp b/Engine/Plugins/Lumin/MagicLeap/Source/MagicLeapAR/Private/LuminARTrackingSystem.cpp
new file mode 100644
--- /dev/null
+++ b/Engine/Plugins/Lumin/MagicLeap/Source/MagicLeapAR/Private/LuminARTrackingSystem.cpp
@@ -0,0 +1,215 @@
+// Copyright 1998-2018 Epic Games, Inc. All Rights Reserved.
+
+#include "LuminARTrackingSystem.h"
+#include "Engine/Engine.h"
+#include "RHIDefinitions.h"
+#include "GameFramework/PlayerController.h"
+#include "GameFramework/WorldSettings.h"
+#include "LuminARDevice.h"
+#include "ARSessionConfig.h"
+
+
+FLuminARImplementation::FLuminARImplementation()
+	: LuminARDeviceInstance(nullptr)
+	, bHasValidPose(false)
+	, DeltaControlRotation(FRotator::ZeroRotator)
+	, DeltaControlOrientation(FQuat::Identity)
+	, LightEstimate(nullptr)
+{
+	LuminARDeviceInstance = FLuminARDevice::GetInstance();
+	check(LuminARDeviceInstance);
+}
+
+FLuminARImplementation::~FLuminARImplementation()
+{
+}
+
+/////////////////////////////////////////////////////////////////////////////////
+// Begin FLuminARImplementation IHeadMountedDisplay Virtual Interface   //
+////////////////////////////////////////////////////////////////////////////////
+
+
+void* FLuminARImplementation::GetARSessionRawPointer()
+{
+#if PLATFORM_LUMIN
+	return static_cast<void*>(FLuminARDevice::GetInstance()->GetARSessionRawPointer());
+#endif
+	ensureAlwaysMsgf(false, TEXT("FLuminARImplementation::GetARSessionRawPointer is unimplemented on current platform."));
+	return nullptr;
+}
+
+void* FLuminARImplementation::GetGameThreadARFrameRawPointer()
+{
+#if PLATFORM_LUMIN
+	return static_cast<void*>(FLuminARDevice::GetInstance()->GetGameThreadARFrameRawPointer());
+#endif
+	ensureAlwaysMsgf(false, TEXT("FLuminARImplementation::GetARSessionRawPointer is unimplemented on current platform."));
+	return nullptr;
+}
+
+void FLuminARImplementation::OnARSystemInitialized()
+{
+}
+
+bool FLuminARImplementation::OnStartARGameFrame(FWorldContext& WorldContext)
+{
+	if (LuminARDeviceInstance->GetIsLuminARSessionRunning())
+	{
+		bHasValidPose = (LuminARDeviceInstance->GetTrackingState() == ELuminARTrackingState::Tracking);
+
+		if (LightEstimate == nullptr)
+		{
+			LightEstimate = NewObject<UARBasicLightEstimate>();
+		}
+		FLuminARLightEstimate LuminARLightEstimate = FLuminARDevice::GetInstance()->GetLatestLightEstimate();
+		if (LuminARLightEstimate.bIsValid)
+		{
+			LightEstimate->SetLightEstimate(LuminARLightEstimate.RGBScaleFactor, LuminARLightEstimate.PixelIntensity);
+		}
+		else
+		{
+			LightEstimate = nullptr;
+		}
+	}
+	return true;
+}
+
+
+EARTrackingQuality FLuminARImplementation::OnGetTrackingQuality() const
+{
+	if (!bHasValidPose)
+	{
+		return EARTrackingQuality::NotTracking;
+	}
+
+	return EARTrackingQuality::OrientationAndPosition;
+}
+
+void FLuminARImplementation::OnStartARSession(UARSessionConfig* SessionConfig)
+{
+	FLuminARDevice::GetInstance()->StartLuminARSessionRequest(SessionConfig);
+}
+
+void FLuminARImplementation::OnPauseARSession()
+{
+	FLuminARDevice::GetInstance()->PauseLuminARSession();
+}
+
+void FLuminARImplementation::OnStopARSession()
+{
+	FLuminARDevice::GetInstance()->PauseLuminARSession();
+	FLuminARDevice::GetInstance()->ResetLuminARSession();
+}
+
+FARSessionStatus FLuminARImplementation::OnGetARSessionStatus() const
+{
+	return FLuminARDevice::GetInstance()->GetSessionStatus();
+}
+
+void FLuminARImplementation::OnSetAlignmentTransform(const FTransform& InAlignmentTransform)
+{
+	const FTransform& NewAlignmentTransform = InAlignmentTransform;
+
+	TArray<UARTrackedGeometry*> AllTrackedGeometries = OnGetAllTrackedGeometries();
+	for (UARTrackedGeometry* TrackedGeometry : AllTrackedGeometries)
+	{
+		TrackedGeometry->UpdateAlignmentTransform(NewAlignmentTransform);
+	}
+
+	TArray<UARPin*> AllARPins = OnGetAllPins();
+	for (UARPin* SomePin : AllARPins)
+	{
+		SomePin->UpdateAlignmentTransform(NewAlignmentTransform);
+	}
+}
+
+static ELuminARLineTraceChannel ConvertToLuminTraceChannels(EARLineTraceChannels TraceChannels)
+{
+	ELuminARLineTraceChannel LuminARTraceChannels = ELuminARLineTraceChannel::None;
+	if (!!(TraceChannels & EARLineTraceChannels::FeaturePoint))
+	{
+		LuminARTraceChannels = LuminARTraceChannels | ELuminARLineTraceChannel::FeaturePoint;
+	}
+	if (!!(TraceChannels & EARLineTraceChannels::GroundPlane))
+	{
+		LuminARTraceChannels = LuminARTraceChannels | ELuminARLineTraceChannel::InfinitePlane;
+	}
+	if (!!(TraceChannels & EARLineTraceChannels::PlaneUsingBoundaryPolygon))
+	{
+		LuminARTraceChannels = LuminARTraceChannels | ELuminARLineTraceChannel::PlaneUsingBoundaryPolygon;
+	}
+	if (!!(TraceChannels & EARLineTraceChannels::PlaneUsingExtent))
+	{
+		LuminARTraceChannels = LuminARTraceChannels | ELuminARLineTraceChannel::PlaneUsingExtent;
+	}
+	return LuminARTraceChannels;
+}
+
+TArray<FARTraceResult> FLuminARImplementation::OnLineTraceTrackedObjects(const FVector2D ScreenCoord, EARLineTraceChannels TraceChannels)
+{
+	TArray<FARTraceResult> OutHitResults;
+	FLuminARDevice::GetInstance()->ARLineTrace(ScreenCoord, ConvertToLuminTraceChannels(TraceChannels), OutHitResults);
+	return OutHitResults;
+}
+
+TArray<FARTraceResult> FLuminARImplementation::OnLineTraceTrackedObjects(const FVector Start, const FVector End, EARLineTraceChannels TraceChannels)
+{
+	TArray<FARTraceResult> OutHitResults;
+	FLuminARDevice::GetInstance()->ARLineTrace(Start, End, ConvertToLuminTraceChannels(TraceChannels), OutHitResults);
+	return OutHitResults;
+}
+
+TArray<UARTrackedGeometry*> FLuminARImplementation::OnGetAllTrackedGeometries() const
+{
+	TArray<UARTrackedGeometry*> AllTrackedGeometry;
+	FLuminARDevice::GetInstance()->GetAllTrackables<UARTrackedGeometry>(AllTrackedGeometry);
+	return AllTrackedGeometry;
+}
+
+TArray<UARPin*> FLuminARImplementation::OnGetAllPins() const
+{
+	TArray<UARPin*> AllARPins;
+	FLuminARDevice::GetInstance()->GetAllARPins(AllARPins);
+	return AllARPins;
+}
+
+bool FLuminARImplementation::OnIsTrackingTypeSupported(EARSessionType SessionType) const
+{
+	return FLuminARDevice::GetInstance()->GetIsTrackingTypeSupported(SessionType);
+}
+
+UARLightEstimate* FLuminARImplementation::OnGetCurrentLightEstimate() const
+{
+	return LightEstimate;
+}
+
+UARPin* FLuminARImplementation::OnPinComponent(USceneComponent* ComponentToPin, const FTransform& PinToWorldTransform, UARTrackedGeometry* TrackedGeometry /*= nullptr*/, const FName DebugName /*= NAME_None*/)
+{
+	UARPin* NewARPin = nullptr;
+	ELuminARFunctionStatus Status = FLuminARDevice::GetInstance()->CreateARPin(PinToWorldTransform, TrackedGeometry, ComponentToPin, DebugName, NewARPin);
+	if (Status != ELuminARFunctionStatus::Success)
+	{
+		UE_LOG(LogLuminARImplementation, Warning, TEXT("OnPinComponent CreateARPin did not return success.  Status=%i"), static_cast<int32>(Status));
+	}
+	return NewARPin;
+}
+
+void FLuminARImplementation::OnRemovePin(UARPin* PinToRemove)
+{
+	FLuminARDevice::GetInstance()->RemoveARPin(PinToRemove);
+}
+
+TArray<FVector> FLuminARImplementation::OnGetPointCloud() const
+{
+	//TODO NEEDED
+	return TArray<FVector>();
+}
+
+
+void FLuminARImplementation::AddReferencedObjects(FReferenceCollector& Collector)
+{
+	if (LightEstimate != nullptr)
+	{
+		Collector.AddReferencedObject(LightEstimate);
+	}
+}
diff --git a/Engine/Plugins/Lumin/MagicLeap/Source/MagicLeapAR/Private/LuminARTrackingSystem.h b/Engine/Plugins/Lumin/MagicLeap/Source/MagicLeapAR/Private/LuminARTrackingSystem.h
new file mode 100644
--- /dev/null
+++ b/Engine/Plugins/Lumin/MagicLeap/Source/MagicLeapAR/Private/LuminARTrackingSystem.h
@@ -0,0 +1,79 @@
+// Copyright 1998-2018 Epic Games, Inc. All Rights Reserved.
+
+#pragma once
+#include "XRTrackingSystemBase.h"
+#include "HeadMountedDisplay.h"
+#include "IHeadMountedDisplay.h"
+#include "SceneViewExtension.h"
+#include "Slate/SceneViewport.h"
+#include "SceneView.h"
+#include "LuminARDevice.h"
+#include "ARSystem.h"
+#include "ARLightEstimate.h"
+#include "Engine/Engine.h" // for FWorldContext
+
+DEFINE_LOG_CATEGORY_STATIC(LogLuminAR, Log, All);
+
+class FLuminARImplementation : public IARSystemSupport, public FGCObject, public TSharedFromThis<IARSystemSupport, ESPMode::ThreadSafe>
+{
+	friend class FLuminARCamera;
+
+public:
+	FLuminARImplementation();
+	~FLuminARImplementation();
+
+	// IXRTrackingSystem
+
+	void* GetARSessionRawPointer() override;
+	void* GetGameThreadARFrameRawPointer() override;
+
+protected:
+	// IARSystemSupport
+	virtual void OnARSystemInitialized() override;
+	virtual bool OnStartARGameFrame(FWorldContext& WorldContext) override;
+
+	virtual EARTrackingQuality OnGetTrackingQuality() const override;
+	virtual void OnStartARSession(UARSessionConfig* SessionConfig) override;
+	virtual void OnPauseARSession() override;
+	virtual void OnStopARSession() override;
+	virtual FARSessionStatus OnGetARSessionStatus() const override;
+	virtual void OnSetAlignmentTransform(const FTransform& InAlignmentTransform) override;
+	virtual TArray<FARTraceResult> OnLineTraceTrackedObjects(const FVector2D ScreenCoord, EARLineTraceChannels TraceChannels) override;
+	virtual TArray<FARTraceResult> OnLineTraceTrackedObjects(const FVector Start, const FVector End, EARLineTraceChannels TraceChannels) override;
+	virtual TArray<UARTrackedGeometry*> OnGetAllTrackedGeometries() const override;
+	virtual TArray<UARPin*> OnGetAllPins() const override;
+	virtual bool OnIsTrackingTypeSupported(EARSessionType SessionType) const override;
+	virtual UARLightEstimate* OnGetCurrentLightEstimate() const override;
+
+	virtual UARPin* OnPinComponent(USceneComponent* ComponentToPin, const FTransform& PinToWorldTransform, UARTrackedGeometry* TrackedGeometry = nullptr, const FName DebugName = NAME_None) override;
+	virtual void OnRemovePin(UARPin* PinToRemove) override;
+	virtual UARTextureCameraImage* OnGetCameraImage() override { return nullptr; }
+	virtual UARTextureCameraDepth* OnGetCameraDepth() override { return nullptr; }
+	virtual bool OnAddManualEnvironmentCaptureProbe(FVector Location, FVector Extent) { return false; }
+	virtual TSharedPtr<FARGetCandidateObjectAsyncTask, ESPMode::ThreadSafe> OnGetCandidateObject(FVector Location, FVector Extent) const { return TSharedPtr<FARGetCandidateObjectAsyncTask, ESPMode::ThreadSafe>(); }
+	virtual TSharedPtr<FARSaveWorldAsyncTask, ESPMode::ThreadSafe> OnSaveWorld() const { return TSharedPtr<FARSaveWorldAsyncTask, ESPMode::ThreadSafe>(); }
+	virtual EARWorldMappingState OnGetWorldMappingStatus() const { return EARWorldMappingState::StillMappingNotRelocalizable; }
+	virtual TArray<FARVideoFormat> OnGetSupportedVideoFormats(EARSessionType SessionType) const override { return TArray<FARVideoFormat>(); }
+	virtual TArray<FVector> OnGetPointCloud() const;
+	virtual bool OnAddRuntimeCandidateImage(UARSessionConfig* SessionConfig, UTexture2D* CandidateTexture, FString FriendlyName, float PhysicalWidth) { return false; }
+	//~IARSystemSupport
+
+private:
+	//~ FGCObject
+	virtual void AddReferencedObjects(FReferenceCollector& Collector) override;
+	//~ FGCObject
+
+private:
+	FLuminARDevice* LuminARDeviceInstance;
+
+	bool bHasValidPose;
+
+	FRotator DeltaControlRotation;    // same as DeltaControlOrientation but as rotator
+	FQuat DeltaControlOrientation; // same as DeltaControlRotation but as quat
+
+	TSharedPtr<class ISceneViewExtension, ESPMode::ThreadSafe> ViewExtension;
+
+	UARBasicLightEstimate* LightEstimate;
+};
+
+DEFINE_LOG_CATEGORY_STATIC(LogLuminARImplementation, Log, All);
diff --git a/Engine/Plugins/Lumin/MagicLeap/Source/MagicLeapAR/Public/LuminARFunctionLibrary.h b/Engine/Plugins/Lumin/MagicLeap/Source/MagicLeapAR/Public/LuminARFunctionLibrary.h
new file mode 100644
--- /dev/null
+++ b/Engine/Plugins/Lumin/MagicLeap/Source/MagicLeapAR/Public/LuminARFunctionLibrary.h
@@ -0,0 +1,82 @@
+// Copyright 1998-2018 Epic Games, Inc. All Rights Reserved.
+
+#pragma once
+
+#include "CoreMinimal.h"
+#include "Engine/LatentActionManager.h"
+#include "Kismet/BlueprintFunctionLibrary.h"
+
+#include "LuminARTypes.h"
+#include "LuminARSessionConfig.h"
+#include "LuminARFunctionLibrary.generated.h"
+
+/** A function library that provides static/Blueprint functions associated with LuminAR session.*/
+UCLASS()
+class MAGICLEAPAR_API ULuminARSessionFunctionLibrary : public UBlueprintFunctionLibrary
+{
+	GENERATED_BODY()
+
+public:
+	//-----------------Lifecycle---------------------
+
+	/**
+	 * Starts a new LuminAR tracking session LuminAR specific configuration.
+	 * If the session already started and the config isn't the same, it will stop the previous session and start a new session with the new config.
+	 * Note that this is a latent action, you can query the session start result by querying GetLuminARSessionStatus() after the latent action finished.
+	 *
+	 * @param Configuration				The LuminARSession configuration to start the session.
+	 */
+	UFUNCTION(BlueprintCallable, Category = "LuminAR|Session", meta = (Latent, LatentInfo = "LatentInfo", WorldContext = "WorldContextObject", Keywords = "luminar session start config"))
+	static void StartLuminARSession(UObject* WorldContextObject, struct FLatentActionInfo LatentInfo, ULuminARSessionConfig* Configuration);
+};
+
+/** A function library that provides static/Blueprint functions associated with most recent LuminAR tracking frame.*/
+UCLASS()
+class MAGICLEAPAR_API ULuminARFrameFunctionLibrary : public UBlueprintFunctionLibrary
+{
+	GENERATED_BODY()
+
+public:
+	/**
+	 * Returns the current ARCore session status.
+	 *
+	 * @return	A EARSessionStatus enum that describes the session status.
+	 */
+	UFUNCTION(BlueprintPure, Category = "LuminAR|MotionTracking", meta = (Keywords = "luminar session"))
+	static ELuminARTrackingState GetTrackingState();
+
+	///**
+	// * Gets the latest tracking pose in Unreal world space of the ARCore device.
+	// *
+	// * Note that ARCore motion tracking has already integrated with HMD and the motion controller interface.
+	// * Use this function only if you need to implement your own tracking component.
+	// *
+	// * @param OutPose		The latest device pose.
+	// * @return				True if the pose is updated successfully for this frame.
+	// */
+	//UFUNCTION(BlueprintPure, Category = "LuminAR|MotionTracking", meta = (Keywords = "luminar pose transform"))
+	//static void GetPose(FTransform& OutPose);
+
+	/**
+	 * Traces a ray from the user's device in the direction of the given location in the camera
+	 * view. Intersections with detected scene geometry are returned, sorted by distance from the
+	 * device; the nearest intersection is returned first.
+	 *
+	 * @param WorldContextObject	The world context.
+	 * @param ScreenPosition		The position on the screen to cast the ray from.
+	 * @param ARObjectType			A set of ELuminARLineTraceChannel indicate which type of line trace it should perform.
+	 * @param OutHitResults			The list of hit results sorted by distance.
+	 * @return						True if there is a hit detected.
+	 */
+	UFUNCTION(BlueprintCallable, Category = "LuminAR|LineTrace", meta = (WorldContext = "WorldContextObject", Keywords = "luminar raycast hit"))
+	static bool LuminARLineTrace(UObject* WorldContextObject, const FVector2D& ScreenPosition, TSet<ELuminARLineTraceChannel> TraceChannels, TArray<FARTraceResult>& OutHitResults);
+
+
+	/**
+	 * Gets the latest light estimation.
+	 *
+	 * @param OutLightEstimate		The struct that describes the latest light estimation.
+	 */
+	UFUNCTION(BlueprintCallable, Category = "LuminAR|LightEstimation", meta = (Keywords = "luminar light ambient"))
+	static void GetLightEstimation(FLuminARLightEstimate& OutLightEstimate);
+};
diff --git a/Engine/Plugins/Lumin/MagicLeap/Source/MagicLeapAR/Public/LuminARModule.h b/Engine/Plugins/Lumin/MagicLeap/Source/MagicLeapAR/Public/LuminARModule.h
new file mode 100644
--- /dev/null
+++ b/Engine/Plugins/Lumin/MagicLeap/Source/MagicLeapAR/Public/LuminARModule.h
@@ -0,0 +1,19 @@
+// Copyright 1998-2018 Epic Games, Inc. All Rights Reserved.
+
+#pragma once
+
+#include "Modules/ModuleManager.h"
+
+/**
+ * The public interface to this module.
+ */
+class ILuminARModule : public IModuleInterface
+{
+public:
+	//create for mutual connection (regardless of construction order)
+	virtual TSharedPtr<IARSystemSupport, ESPMode::ThreadSafe> CreateARImplementation() = 0;
+	//Now connect (regardless of connection order)
+	virtual void ConnectARImplementationToXRSystem(FXRTrackingSystemBase* InXRTrackingSystem) = 0;
+	//Now initialize fully connected systems
+	virtual void InitializeARImplementation() = 0;
+};
diff --git a/Engine/Plugins/Lumin/MagicLeap/Source/MagicLeapAR/Public/LuminARSessionConfig.h b/Engine/Plugins/Lumin/MagicLeap/Source/MagicLeapAR/Public/LuminARSessionConfig.h
new file mode 100644
--- /dev/null
+++ b/Engine/Plugins/Lumin/MagicLeap/Source/MagicLeapAR/Public/LuminARSessionConfig.h
@@ -0,0 +1,66 @@
+// Copyright 1998-2018 Epic Games, Inc. All Rights Reserved.
+
+#pragma once
+
+#include "UObject/Object.h"
+#include "UObject/ObjectMacros.h"
+#include "LuminARTypes.h"
+#include "ARSessionConfig.h"
+
+#include "LuminARSessionConfig.generated.h"
+
+
+UCLASS(BlueprintType, Category = "AR AugmentedReality")
+class MAGICLEAPAR_API ULuminARSessionConfig : public UARSessionConfig
+{
+	GENERATED_BODY()
+
+public:
+	ULuminARSessionConfig()
+		: UARSessionConfig()
+		, PlaneSearchExtents(10000.0f, 10000.0f, 10000.0f)
+	{
+	}
+
+	/** The maximum number of plane results that will be returned. */
+	UPROPERTY(EditAnywhere, Category = "Lumin AR Settings")
+	int32 MaxPlaneQueryResults = 200;
+
+	/** The minimum area (in square cm) of planes to be returned. This value cannot be lower than 400 (lower values will be capped to this minimum). A good default value is 2500. */
+	UPROPERTY(EditAnywhere, Category = "Lumin AR Settings")
+	int32 MinPlaneArea = 400;
+
+	/** Should we detect planes with any orientation (ie not just horizontal or vertical). */
+	UPROPERTY(EditAnywhere, Category = "Lumin AR Settings")
+	bool bArbitraryOrientationPlaneDetection = false; // default to false, for now anyway, because some other platforms do not support this.
+
+	//UPROPERTY(EditAnywhere, Category = "Lumin AR Settings")
+	//bool bGetInnerPlanes = false;
+
+	/** The dimensions of the box within which plane results will be returned.  The box center and rotation are those of the tracking to world transform origin. */
+	UPROPERTY(EditAnywhere, Category = "Lumin AR Settings")
+	FVector PlaneSearchExtents;
+
+	/** Additional Flags to apply to the plane queries. Note: the plane orientation detection settings also cause flags to be set.  It is ok to duplicate those here.*/
+	UPROPERTY(EditAnywhere, Category = "Lumin AR Settings")
+	TArray<EPlaneQueryFlags> PlaneQueryFlags;
+
+	/** If true discard any 'plane' objects that come through with zero extents and only polygon edge data.*/
+	UPROPERTY(EditAnywhere, Category = "Lumin AR Settings")
+	bool bDiscardZeroExtentPlanes = true;
+
+	
+	static ULuminARSessionConfig* CreateARCoreSessionConfig(bool bHorizontalPlaneDetection, bool bVerticalPlaneDetection, EARLightEstimationMode InLightEstimationMode, EARFrameSyncMode InFrameSyncMode, bool bInEnableAutomaticCameraOverlay, bool bInEnableAutomaticCameraTracking)
+	{
+		ULuminARSessionConfig* NewARSessionConfig = NewObject<ULuminARSessionConfig>();
+
+		NewARSessionConfig->bHorizontalPlaneDetection = bHorizontalPlaneDetection;
+		NewARSessionConfig->bHorizontalPlaneDetection = bVerticalPlaneDetection;
+		NewARSessionConfig->LightEstimationMode = InLightEstimationMode;
+		NewARSessionConfig->FrameSyncMode = InFrameSyncMode;
+		NewARSessionConfig->bEnableAutomaticCameraOverlay = bInEnableAutomaticCameraOverlay;
+		NewARSessionConfig->bEnableAutomaticCameraTracking = bInEnableAutomaticCameraTracking;
+		
+		return NewARSessionConfig;
+	};
+};
diff --git a/Engine/Plugins/Lumin/MagicLeap/Source/MagicLeapAR/Public/LuminARTypes.h b/Engine/Plugins/Lumin/MagicLeap/Source/MagicLeapAR/Public/LuminARTypes.h
new file mode 100644
--- /dev/null
+++ b/Engine/Plugins/Lumin/MagicLeap/Source/MagicLeapAR/Public/LuminARTypes.h
@@ -0,0 +1,201 @@
+// Copyright 1998-2018 Epic Games, Inc. All Rights Reserved.
+#pragma once
+
+#include "CoreMinimal.h"
+
+#if WITH_MLSDK
+#include <ml_api.h>
+#endif //WITH_MLSDK
+
+#include "ARTypes.h"
+#include "ARTrackable.h"
+#include "ARTraceResult.h"
+#include "ARSystem.h"
+#include "ARPin.h"
+
+#include "LuminARTypes.generated.h"
+
+/// @defgroup LuminARBase
+/// The base module for LuminAR plugin
+
+#if PLATFORM_LUMIN
+
+typedef enum
+{
+	LUMIN_AR_TRACKABLE_NOT_VALID,
+	LUMIN_AR_TRACKABLE_PLANE,
+	LUMIN_AR_TRACKABLE_POINT
+}
+ArTrackableType;
+
+//TODO - Proper data?
+struct ArPose
+{
+	FVector Pos;
+	FQuat Quat;
+};
+
+struct LuminArTrackable
+{
+	LuminArTrackable(ArPose InPose)
+		: Handle(ML_INVALID_HANDLE)
+	{
+	}
+
+	MLHandle Handle;
+};
+
+struct LuminArAnchor : public LuminArTrackable
+{
+	LuminArAnchor(ArPose InPose, MLHandle InParentTrackable)
+		: LuminArTrackable(InPose)
+		, ParentTrackable(InParentTrackable)
+	{}
+
+	void Detach()
+	{
+		ParentTrackable = ML_INVALID_HANDLE;
+	}
+
+	MLHandle ParentTrackable;
+};
+
+typedef struct MLPlane ArPlane;
+typedef struct MLPlane ArPoint;
+
+//TODO - Add details for ArImage
+struct ArImage
+{
+};
+
+#endif
+
+UENUM(BlueprintType)
+enum class ELuminARAvailability : uint8
+{
+	/* An internal error occurred while determining Lumin AR availability. */
+	UnknownError = 0,
+
+	/* ARCore is supported, installed, and available to use. */
+	SupportedInstalled = 200
+};
+
+/**
+ * @ingroup LuminARBase
+ * Describes the status of most LuminAR functions.
+ */
+UENUM(BlueprintType)
+enum class ELuminARFunctionStatus : uint8
+{
+	/** Function returned successfully. */
+	Success,
+	/** Function failed due to Fatal error. */
+	Fatal,
+	/** Function failed due to the session isn't running. */
+	SessionPaused,
+	/** Function failed due to ARCore session isn't in tracking state. */
+	NotTracking,
+	/** Function failed due to the requested resource is exhausted. */
+	ResourceExhausted,
+	/** Function failed due to the requested resource isn't available yet. */
+	NotAvailable,
+	/** Function failed due to the function augment has invalid type. */
+	InvalidType,
+	/** Function failed with unknown reason. */
+	Unknown
+};
+
+/**
+ * @ingroup LuminARBase
+ * Describes the tracking state of the current ARCore session.
+ */
+UENUM(BlueprintType)
+enum class ELuminARTrackingState : uint8
+{
+	/** Tracking is valid. */
+	Tracking = 0,
+	/** Tracking is temporary lost but could recover in the future. */
+	NotTracking = 1,
+	/** Tracking is lost will not recover. */
+	StoppedTracking = 2
+};
+
+/**
+ * A struct describes the ARCore light estimation.
+ */
+USTRUCT(BlueprintType)
+struct FLuminARLightEstimate
+{
+	GENERATED_BODY()
+
+	/** Whether this light estimation is valid. */
+	UPROPERTY(BlueprintReadOnly, Category = "LuminAR|LightEstimate")
+	bool bIsValid;
+
+	/** The average pixel intensity of the passthrough camera image. */
+	UPROPERTY(BlueprintReadOnly, Category = "LuminAR|LightEstimate")
+	float PixelIntensity;
+	
+	/**
+	 * The RGB scale to match the color of the light in the real environment.
+	 */
+	UPROPERTY(BlueprintReadOnly, Category = "LuminAR|LightEstimate")
+	FVector RGBScaleFactor;
+};
+
+/**
+ * @ingroup LuminARBase
+ * Describes which channel ARLineTrace will be performed on.
+ */
+UENUM(BlueprintType, Category = "LuminAR|TraceChannel", meta = (Bitflags))
+enum class ELuminARLineTraceChannel : uint8
+{
+	None = 0,
+	/** Trace against feature point cloud. */
+	FeaturePoint = 1,
+	/** Trace against the infinite plane. */
+	InfinitePlane = 2,
+	/** Trace against the plane using its extent. */
+	PlaneUsingExtent = 4,
+	/** Trace against the plane using its boundary polygon. */
+	PlaneUsingBoundaryPolygon = 8,
+	/**
+	 * Trace against feature point and attempt to estimate the normal of the surface centered around the trace hit point.
+	 * Surface normal estimation is most likely to succeed on textured surfaces and with camera motion.
+	 */
+	FeaturePointWithSurfaceNormal = 16,
+};
+ENUM_CLASS_FLAGS(ELuminARLineTraceChannel);
+
+
+/** Control flags for plane queries. */
+UENUM(BlueprintType)
+enum class EPlaneQueryFlags : uint8
+{
+	/** Include planes whose normal is perpendicular to gravity. */
+	Vertical,
+
+	/** Include planes whose normal is parallel to gravity. */
+	Horizontal,
+
+	/** Include planes with arbitrary normals. */
+	Arbitrary,
+
+	/** If set, non-horizontal planes will be aligned perpendicular to gravity. */
+	OrientToGravity,
+
+	/** If set, inner planes will be returned; if not set, outer planes will be returned. */
+	PreferInner,
+
+	/** If set, holes in planar surfaces will be ignored. */
+	IgnoreHoles,
+
+	/** If set, include planes semantically tagged as ceiling. */
+	Ceiling,
+
+	/** If set, include planes semantically tagged as floor. */
+	Floor,
+
+	/** If set, include planes semantically tagged as wall. */
+	Wall
+};
\ No newline at end of file
diff --git a/Engine/Plugins/Runtime/AR/Apple/AppleARKit/Source/AppleARKit/Private/AppleARKitSystem.cpp b/Engine/Plugins/Runtime/AR/Apple/AppleARKit/Source/AppleARKit/Private/AppleARKitSystem.cpp
--- a/Engine/Plugins/Runtime/AR/Apple/AppleARKit/Source/AppleARKit/Private/AppleARKitSystem.cpp
+++ b/Engine/Plugins/Runtime/AR/Apple/AppleARKit/Source/AppleARKit/Private/AppleARKitSystem.cpp
@@ -682,6 +682,13 @@ TArray<FARTraceResult> FAppleARKitSystem::OnLineTraceTrackedObjects( const FVect
 	return Results;
 }
 
+TArray<FARTraceResult> FAppleARKitSystem::OnLineTraceTrackedObjects(const FVector Start, const FVector End, EARLineTraceChannels TraceChannels)
+{
+	UE_LOG(LogAppleARKit, Warning, TEXT("FAppleARKitSystem::OnLineTraceTrackedObjects(Start, End, TraceChannels) is currently unsupported.  No results will be returned."))
+	TArray<FARTraceResult> EmptyResults;
+	return EmptyResults;
+}
+
 TArray<UARTrackedGeometry*> FAppleARKitSystem::OnGetAllTrackedGeometries() const
 {
 	TArray<UARTrackedGeometry*> Geometries;
diff --git a/Engine/Plugins/Runtime/AR/Apple/AppleARKit/Source/AppleARKit/Private/AppleARKitSystem.h b/Engine/Plugins/Runtime/AR/Apple/AppleARKit/Source/AppleARKit/Private/AppleARKitSystem.h
--- a/Engine/Plugins/Runtime/AR/Apple/AppleARKit/Source/AppleARKit/Private/AppleARKitSystem.h
+++ b/Engine/Plugins/Runtime/AR/Apple/AppleARKit/Source/AppleARKit/Private/AppleARKitSystem.h
@@ -70,6 +70,7 @@ protected:
 	virtual FARSessionStatus OnGetARSessionStatus() const override;
 	virtual void OnSetAlignmentTransform(const FTransform& InAlignmentTransform) override;
 	virtual TArray<FARTraceResult> OnLineTraceTrackedObjects( const FVector2D ScreenCoord, EARLineTraceChannels TraceChannels ) override;
+	virtual TArray<FARTraceResult> OnLineTraceTrackedObjects(const FVector Start, const FVector End, EARLineTraceChannels TraceChannels) override;
 	virtual TArray<UARTrackedGeometry*> OnGetAllTrackedGeometries() const override;
 	virtual TArray<UARPin*> OnGetAllPins() const override;
 	virtual bool OnIsTrackingTypeSupported(EARSessionType SessionType) const override;
diff --git a/Engine/Plugins/Runtime/AR/Google/GoogleARCore/Source/GoogleARCoreBase/Private/GoogleARCoreXRTrackingSystem.cpp b/Engine/Plugins/Runtime/AR/Google/GoogleARCore/Source/GoogleARCoreBase/Private/GoogleARCoreXRTrackingSystem.cpp
--- a/Engine/Plugins/Runtime/AR/Google/GoogleARCore/Source/GoogleARCoreBase/Private/GoogleARCoreXRTrackingSystem.cpp
+++ b/Engine/Plugins/Runtime/AR/Google/GoogleARCore/Source/GoogleARCoreBase/Private/GoogleARCoreXRTrackingSystem.cpp
@@ -232,7 +232,7 @@ void FGoogleARCoreXRTrackingSystem::OnSetAlignmentTransform(const FTransform& In
 	}
 }
 
-TArray<FARTraceResult> FGoogleARCoreXRTrackingSystem::OnLineTraceTrackedObjects(const FVector2D ScreenCoord, EARLineTraceChannels TraceChannels)
+static EGoogleARCoreLineTraceChannel ConvertToGoogleARCoreTraceChannels(EARLineTraceChannels TraceChannels)
 {
 	EGoogleARCoreLineTraceChannel ARCoreTraceChannels = EGoogleARCoreLineTraceChannel::None;
 	if (!!(TraceChannels & EARLineTraceChannels::FeaturePoint))
@@ -251,9 +251,20 @@ TArray<FARTraceResult> FGoogleARCoreXRTrackingSystem::OnLineTraceTrackedObjects(
 	{
 		ARCoreTraceChannels = ARCoreTraceChannels | EGoogleARCoreLineTraceChannel::PlaneUsingExtent;
 	}
+	return ARCoreTraceChannels;
+}
 
+TArray<FARTraceResult> FGoogleARCoreXRTrackingSystem::OnLineTraceTrackedObjects(const FVector2D ScreenCoord, EARLineTraceChannels TraceChannels)
+{
+	TArray<FARTraceResult> OutHitResults;
+	FGoogleARCoreDevice::GetInstance()->ARLineTrace(ScreenCoord, ConvertToGoogleARCoreTraceChannels(TraceChannels), OutHitResults);
+	return OutHitResults;
+}
+
+TArray<FARTraceResult> FGoogleARCoreXRTrackingSystem::OnLineTraceTrackedObjects(const FVector Start, const FVector End, EARLineTraceChannels TraceChannels)
+{
 	TArray<FARTraceResult> OutHitResults;
-	FGoogleARCoreDevice::GetInstance()->ARLineTrace(ScreenCoord, ARCoreTraceChannels, OutHitResults);
+	FGoogleARCoreDevice::GetInstance()->ARLineTrace(Start, End, ConvertToGoogleARCoreTraceChannels(TraceChannels), OutHitResults);
 	return OutHitResults;
 }
 
diff --git a/Engine/Plugins/Runtime/AR/Google/GoogleARCore/Source/GoogleARCoreBase/Private/GoogleARCoreXRTrackingSystem.h b/Engine/Plugins/Runtime/AR/Google/GoogleARCore/Source/GoogleARCoreBase/Private/GoogleARCoreXRTrackingSystem.h
--- a/Engine/Plugins/Runtime/AR/Google/GoogleARCore/Source/GoogleARCoreBase/Private/GoogleARCoreXRTrackingSystem.h
+++ b/Engine/Plugins/Runtime/AR/Google/GoogleARCore/Source/GoogleARCoreBase/Private/GoogleARCoreXRTrackingSystem.h
@@ -56,6 +56,7 @@ protected:
 	virtual FARSessionStatus OnGetARSessionStatus() const override;
 	virtual void OnSetAlignmentTransform(const FTransform& InAlignmentTransform) override;
 	virtual TArray<FARTraceResult> OnLineTraceTrackedObjects(const FVector2D ScreenCoord, EARLineTraceChannels TraceChannels) override;
+	virtual TArray<FARTraceResult> OnLineTraceTrackedObjects(const FVector Start, const FVector End, EARLineTraceChannels TraceChannels) override;
 	virtual TArray<UARTrackedGeometry*> OnGetAllTrackedGeometries() const override;
 	virtual TArray<UARPin*> OnGetAllPins() const override;
 	virtual bool OnIsTrackingTypeSupported(EARSessionType SessionType) const override;
diff --git a/Engine/Source/Runtime/AugmentedReality/Private/ARBlueprintLibrary.cpp b/Engine/Source/Runtime/AugmentedReality/Private/ARBlueprintLibrary.cpp
--- a/Engine/Source/Runtime/AugmentedReality/Private/ARBlueprintLibrary.cpp
+++ b/Engine/Source/Runtime/AugmentedReality/Private/ARBlueprintLibrary.cpp
@@ -122,6 +122,25 @@ TArray<FARTraceResult> UARBlueprintLibrary::LineTraceTrackedObjects( const FVect
 	return Result;
 }
 
+TArray<FARTraceResult> UARBlueprintLibrary::LineTraceTrackedObjects3D(const FVector Start, const FVector End, bool bTestFeaturePoints, bool bTestGroundPlane, bool bTestPlaneExtents, bool bTestPlaneBoundaryPolygon)
+{
+	TArray<FARTraceResult> Result;
+
+	auto ARSystem = GetARSystem();
+	if (ARSystem.IsValid())
+	{
+		EARLineTraceChannels ActiveTraceChannels =
+			(bTestFeaturePoints ? EARLineTraceChannels::FeaturePoint : EARLineTraceChannels::None) |
+			(bTestGroundPlane ? EARLineTraceChannels::GroundPlane : EARLineTraceChannels::None) |
+			(bTestPlaneExtents ? EARLineTraceChannels::PlaneUsingExtent : EARLineTraceChannels::None) |
+			(bTestPlaneBoundaryPolygon ? EARLineTraceChannels::PlaneUsingBoundaryPolygon : EARLineTraceChannels::None);
+
+		Result = ARSystem->LineTraceTrackedObjects(Start, End, ActiveTraceChannels);
+	}
+
+	return Result;
+}
+
 TArray<UARTrackedGeometry*> UARBlueprintLibrary::GetAllGeometries()
 {
 	TArray<UARTrackedGeometry*> Geometries;
diff --git a/Engine/Source/Runtime/AugmentedReality/Private/ARSupportInterface.cpp b/Engine/Source/Runtime/AugmentedReality/Private/ARSupportInterface.cpp
--- a/Engine/Source/Runtime/AugmentedReality/Private/ARSupportInterface.cpp
+++ b/Engine/Source/Runtime/AugmentedReality/Private/ARSupportInterface.cpp
@@ -150,6 +150,16 @@ TArray<FARTraceResult> FARSupportInterface ::LineTraceTrackedObjects(const FVect
 	return TArray<FARTraceResult>();
 }
 
+TArray<FARTraceResult> FARSupportInterface::LineTraceTrackedObjects(const FVector Start, const FVector End, EARLineTraceChannels TraceChannels)
+{
+	if (ARImplemention)
+	{
+		return ARImplemention->OnLineTraceTrackedObjects(Start, End, TraceChannels);
+	}
+	return TArray<FARTraceResult>();
+
+}
+
 TArray<UARTrackedGeometry*> FARSupportInterface ::GetAllTrackedGeometries() const
 {
 	if (ARImplemention)
diff --git a/Engine/Source/Runtime/AugmentedReality/Private/ARTraceResult.cpp b/Engine/Source/Runtime/AugmentedReality/Private/ARTraceResult.cpp
--- a/Engine/Source/Runtime/AugmentedReality/Private/ARTraceResult.cpp
+++ b/Engine/Source/Runtime/AugmentedReality/Private/ARTraceResult.cpp
@@ -29,6 +29,11 @@ float FARTraceResult::GetDistanceFromCamera() const
 	return DistanceFromCamera;
 }
 
+void FARTraceResult::SetLocalToWorldTransform(const FTransform& LocalToWorldTransform)
+{
+	LocalToTrackingTransform = LocalToWorldTransform * ARSystem->GetXRTrackingSystem()->GetTrackingToWorldTransform().Inverse();
+}
+
 FTransform FARTraceResult::GetLocalToTrackingTransform() const
 {
 	return LocalToTrackingTransform;
diff --git a/Engine/Source/Runtime/AugmentedReality/Public/ARBlueprintLibrary.h b/Engine/Source/Runtime/AugmentedReality/Public/ARBlueprintLibrary.h
--- a/Engine/Source/Runtime/AugmentedReality/Public/ARBlueprintLibrary.h
+++ b/Engine/Source/Runtime/AugmentedReality/Public/ARBlueprintLibrary.h
@@ -74,7 +74,17 @@ public:
 	UFUNCTION(BlueprintCallable, Category = "AR AugmentedReality|Trace Result", meta = (AdvancedDisplay="1", Keywords = "ar augmentedreality augmented reality tracking tracing linetrace"))
 	static TArray<FARTraceResult> LineTraceTrackedObjects( const FVector2D ScreenCoord, bool bTestFeaturePoints = true, bool bTestGroundPlane = true, bool bTestPlaneExtents = true, bool bTestPlaneBoundaryPolygon = true );
 	
-	
+	/**
+	 * Perform a line trace against any real-world geometry as tracked by the AR system.
+	 *
+	 * @param Start					Start point of the trace, in world space.
+	 * @param End					End point of the trace, in world space.
+	 *
+	 * @return a list of \c FARTraceResult sorted by distance from camera.
+	 */
+	UFUNCTION(BlueprintCallable, Category = "AR AugmentedReality|Trace Result", meta = (AdvancedDisplay = "1", Keywords = "ar augmentedreality augmented reality tracking tracing linetrace"))
+	static TArray<FARTraceResult> LineTraceTrackedObjects3D(const FVector Start, const FVector End, bool bTestFeaturePoints = true, bool bTestGroundPlane = true, bool bTestPlaneExtents = true, bool bTestPlaneBoundaryPolygon = true);
+
 	/** @return how well the tracking system is performing at the moment */
 	UFUNCTION(BlueprintPure, Category = "AR AugmentedReality|Tracking", meta = (DisplayName="Get AR Tracking Quality", Keywords = "ar augmentedreality augmented reality tracking quality"))
 	static EARTrackingQuality GetTrackingQuality();
diff --git a/Engine/Source/Runtime/AugmentedReality/Public/ARSessionConfig.h b/Engine/Source/Runtime/AugmentedReality/Public/ARSessionConfig.h
--- a/Engine/Source/Runtime/AugmentedReality/Public/ARSessionConfig.h
+++ b/Engine/Source/Runtime/AugmentedReality/Public/ARSessionConfig.h
@@ -224,6 +224,10 @@ public:
 	/** @see FaceTrackingUpdate */
 	UFUNCTION(BlueprintCallable, Category = "AR Settings")
 	void SetFaceTrackingUpdate(EARFaceTrackingUpdate InUpdate);
+
+
+	bool ShouldDoHorizontalPlaneDetection() const { return bHorizontalPlaneDetection; }
+	bool ShouldDoVerticalPlaneDetection() const { return bVerticalPlaneDetection; }
 	
 	const TArray<uint8>& GetSerializedARCandidateImageDatabase() const;	
 private:
diff --git a/Engine/Source/Runtime/AugmentedReality/Public/ARSupportInterface.h b/Engine/Source/Runtime/AugmentedReality/Public/ARSupportInterface.h
--- a/Engine/Source/Runtime/AugmentedReality/Public/ARSupportInterface.h
+++ b/Engine/Source/Runtime/AugmentedReality/Public/ARSupportInterface.h
@@ -60,6 +60,8 @@ public:
 
 	/** \see UARBlueprintLibrary::LineTraceTrackedObjects() */
 	TArray<FARTraceResult> LineTraceTrackedObjects(const FVector2D ScreenCoords, EARLineTraceChannels TraceChannels);
+	/** \see UARBlueprintLibrary::LineTraceTrackedObjects() */
+	TArray<FARTraceResult> LineTraceTrackedObjects(const FVector Start, const FVector End, EARLineTraceChannels TraceChannels);
 	/** \see UARBlueprintLibrary::GetAllTrackedGeometries() */
 	TArray<UARTrackedGeometry*> GetAllTrackedGeometries() const;
 	/** \see UARBlueprintLibrary::GetAllPins() */
diff --git a/Engine/Source/Runtime/AugmentedReality/Public/ARSystem.h b/Engine/Source/Runtime/AugmentedReality/Public/ARSystem.h
--- a/Engine/Source/Runtime/AugmentedReality/Public/ARSystem.h
+++ b/Engine/Source/Runtime/AugmentedReality/Public/ARSystem.h
@@ -87,7 +87,8 @@ public:
 	 * @return a list of all the geometries that were hit, sorted by distance
 	 */
 	virtual TArray<FARTraceResult> OnLineTraceTrackedObjects( const FVector2D ScreenCoord, EARLineTraceChannels TraceChannels ) = 0;
-	
+	virtual TArray<FARTraceResult> OnLineTraceTrackedObjects( const FVector Start, const FVector End, EARLineTraceChannels TraceChannels ) = 0;
+
 	/** @return a TArray of all the tracked geometries known to your ar system */
 	virtual TArray<UARTrackedGeometry*> OnGetAllTrackedGeometries() const = 0;
 	
diff --git a/Engine/Source/Runtime/AugmentedReality/Public/ARTraceResult.h b/Engine/Source/Runtime/AugmentedReality/Public/ARTraceResult.h
--- a/Engine/Source/Runtime/AugmentedReality/Public/ARTraceResult.h
+++ b/Engine/Source/Runtime/AugmentedReality/Public/ARTraceResult.h
@@ -20,6 +20,8 @@ struct AUGMENTEDREALITY_API FARTraceResult
 	FARTraceResult( const TSharedPtr<FARSupportInterface , ESPMode::ThreadSafe>& InARSystem, float InDistanceFromCamera, EARLineTraceChannels InTraceChannel, const FTransform& InLocalToTrackingTransform, UARTrackedGeometry* InTrackedGeometry );
 	
 	float GetDistanceFromCamera() const;
+
+	void SetLocalToWorldTransform(const FTransform& LocalToWorldTransform);
 	
 	FTransform GetLocalToTrackingTransform() const;
 	
@@ -28,6 +30,14 @@ struct AUGMENTEDREALITY_API FARTraceResult
 	UARTrackedGeometry* GetTrackedGeometry() const;
 	
 	EARLineTraceChannels GetTraceChannel() const;
+
+	struct FARTraceResultComparer
+	{
+		FORCEINLINE_STATS bool operator()(const FARTraceResult& A, const FARTraceResult& B) const
+		{
+			return A.GetDistanceFromCamera() < B.GetDistanceFromCamera();
+		}
+	};
 	
 private:
 	
