LuminARModule forward declares

diff --git a/Engine/Plugins/Lumin/MagicLeap/Source/MagicLeapAR/Public/LuminARModule.h b/Engine/Plugins/Lumin/MagicLeap/Source/MagicLeapAR/Public/LuminARModule.h
--- a/Engine/Plugins/Lumin/MagicLeap/Source/MagicLeapAR/Public/LuminARModule.h
+++ b/Engine/Plugins/Lumin/MagicLeap/Source/MagicLeapAR/Public/LuminARModule.h
@@ -4,6 +4,9 @@
 
 #include "Modules/ModuleManager.h"
 
+class IARSystemSupport;
+class FXRTrackingSystemBase;
+
 /**
  * The public interface to this module.
  */
