- Fix Mac build
- Circular dependency between MagicLeap module and MagicLeapAR module
- Moved ILuminARModule definition to MagicLeapHMD.h in Magic Leap Module
- Moved EPlaneQueryFlags back into PlanesComponent.h in Magic Leap Module
- Removed MagicLeap module dependency on MagicLeapAR
- Probably should clean this up in the future with a MagicLeapTypes module that other MagicLeap modules depend on

diff --git a/Engine/Plugins/Lumin/MagicLeap/Source/MagicLeap/Classes/PlanesComponent.h b/Engine/Plugins/Lumin/MagicLeap/Source/MagicLeap/Classes/PlanesComponent.h
--- a/Engine/Plugins/Lumin/MagicLeap/Source/MagicLeap/Classes/PlanesComponent.h
+++ b/Engine/Plugins/Lumin/MagicLeap/Source/MagicLeap/Classes/PlanesComponent.h
@@ -7,9 +7,40 @@
 #endif //WITH_MLSDK
 
 #include "Components/SceneComponent.h"
-#include "LuminARTypes.h"
 #include "PlanesComponent.generated.h"
 
+/** Control flags for plane queries. */
+UENUM(BlueprintType)
+enum class EPlaneQueryFlags : uint8
+{
+	/** Include planes whose normal is perpendicular to gravity. */
+	Vertical,
+
+	/** Include planes whose normal is parallel to gravity. */
+	Horizontal,
+
+	/** Include planes with arbitrary normals. */
+	Arbitrary,
+
+	/** If set, non-horizontal planes will be aligned perpendicular to gravity. */
+	OrientToGravity,
+
+	/** If set, inner planes will be returned; if not set, outer planes will be returned. */
+	PreferInner,
+
+	/** If set, holes in planar surfaces will be ignored. */
+	IgnoreHoles,
+
+	/** If set, include planes semantically tagged as ceiling. */
+	Ceiling,
+
+	/** If set, include planes semantically tagged as floor. */
+	Floor,
+
+	/** If set, include planes semantically tagged as wall. */
+	Wall
+};
+
 #if WITH_MLSDK
 
 MLPlanesQueryFlags UnrealToMLPlanesQueryFlagMap(EPlaneQueryFlags QueryFlag);
diff --git a/Engine/Plugins/Lumin/MagicLeap/Source/MagicLeap/MagicLeap.Build.cs b/Engine/Plugins/Lumin/MagicLeap/Source/MagicLeap/MagicLeap.Build.cs
--- a/Engine/Plugins/Lumin/MagicLeap/Source/MagicLeap/MagicLeap.Build.cs
+++ b/Engine/Plugins/Lumin/MagicLeap/Source/MagicLeap/MagicLeap.Build.cs
@@ -52,7 +52,6 @@ namespace UnrealBuildTool.Rules
 					"MagicLeapHelperVulkan",
 					"LuminRuntimeSettings",
 					"MagicLeapSecureStorage",
-					"MagicLeapAR",
 				}
 			);
 
diff --git a/Engine/Plugins/Lumin/MagicLeap/Source/MagicLeap/Private/MagicLeapHMD.cpp b/Engine/Plugins/Lumin/MagicLeap/Source/MagicLeap/Private/MagicLeapHMD.cpp
--- a/Engine/Plugins/Lumin/MagicLeap/Source/MagicLeap/Private/MagicLeapHMD.cpp
+++ b/Engine/Plugins/Lumin/MagicLeap/Source/MagicLeap/Private/MagicLeapHMD.cpp
@@ -43,7 +43,6 @@
 #include "VulkanRHIBridge.h"
 #endif
 
-#include "LuminARModule.h"
 #include "MagicLeapPluginUtil.h" // for ML_INCLUDES_START/END
 
 #if WITH_MLSDK
diff --git a/Engine/Plugins/Lumin/MagicLeap/Source/MagicLeap/Private/MagicLeapHMD.h b/Engine/Plugins/Lumin/MagicLeap/Source/MagicLeap/Private/MagicLeapHMD.h
--- a/Engine/Plugins/Lumin/MagicLeap/Source/MagicLeap/Private/MagicLeapHMD.h
+++ b/Engine/Plugins/Lumin/MagicLeap/Source/MagicLeap/Private/MagicLeapHMD.h
@@ -25,6 +25,23 @@ ML_INCLUDES_START
 ML_INCLUDES_END
 #endif //WITH_MLSDK
 
+class IARSystemSupport;
+class FXRTrackingSystemBase;
+
+/**
+ * The public interface to this module.
+ */
+class ILuminARModule : public IModuleInterface
+{
+public:
+	//create for mutual connection (regardless of construction order)
+	virtual TSharedPtr<IARSystemSupport, ESPMode::ThreadSafe> CreateARImplementation() = 0;
+	//Now connect (regardless of connection order)
+	virtual void ConnectARImplementationToXRSystem(FXRTrackingSystemBase* InXRTrackingSystem) = 0;
+	//Now initialize fully connected systems
+	virtual void InitializeARImplementation() = 0;
+};
+
 /**
   * MagicLeap Head Mounted Display
   */
diff --git a/Engine/Plugins/Lumin/MagicLeap/Source/MagicLeapAR/Public/LuminARModule.h b/Engine/Plugins/Lumin/MagicLeap/Source/MagicLeapAR/Public/LuminARModule.h
--- a/Engine/Plugins/Lumin/MagicLeap/Source/MagicLeapAR/Public/LuminARModule.h
+++ b/Engine/Plugins/Lumin/MagicLeap/Source/MagicLeapAR/Public/LuminARModule.h
@@ -3,20 +3,4 @@
 #pragma once
 
 #include "Modules/ModuleManager.h"
-
-class IARSystemSupport;
-class FXRTrackingSystemBase;
-
-/**
- * The public interface to this module.
- */
-class ILuminARModule : public IModuleInterface
-{
-public:
-	//create for mutual connection (regardless of construction order)
-	virtual TSharedPtr<IARSystemSupport, ESPMode::ThreadSafe> CreateARImplementation() = 0;
-	//Now connect (regardless of connection order)
-	virtual void ConnectARImplementationToXRSystem(FXRTrackingSystemBase* InXRTrackingSystem) = 0;
-	//Now initialize fully connected systems
-	virtual void InitializeARImplementation() = 0;
-};
+#include "MagicLeapHMD.h"
diff --git a/Engine/Plugins/Lumin/MagicLeap/Source/MagicLeapAR/Public/LuminARSessionConfig.h b/Engine/Plugins/Lumin/MagicLeap/Source/MagicLeapAR/Public/LuminARSessionConfig.h
--- a/Engine/Plugins/Lumin/MagicLeap/Source/MagicLeapAR/Public/LuminARSessionConfig.h
+++ b/Engine/Plugins/Lumin/MagicLeap/Source/MagicLeapAR/Public/LuminARSessionConfig.h
@@ -5,6 +5,7 @@
 #include "UObject/Object.h"
 #include "UObject/ObjectMacros.h"
 #include "LuminARTypes.h"
+#include "PlanesComponent.h"
 #include "ARSessionConfig.h"
 
 #include "LuminARSessionConfig.generated.h"
diff --git a/Engine/Plugins/Lumin/MagicLeap/Source/MagicLeapAR/Public/LuminARTypes.h b/Engine/Plugins/Lumin/MagicLeap/Source/MagicLeapAR/Public/LuminARTypes.h
--- a/Engine/Plugins/Lumin/MagicLeap/Source/MagicLeapAR/Public/LuminARTypes.h
+++ b/Engine/Plugins/Lumin/MagicLeap/Source/MagicLeapAR/Public/LuminARTypes.h
@@ -168,34 +168,3 @@ enum class ELuminARLineTraceChannel : uint8
 ENUM_CLASS_FLAGS(ELuminARLineTraceChannel);
 
 
-/** Control flags for plane queries. */
-UENUM(BlueprintType)
-enum class EPlaneQueryFlags : uint8
-{
-	/** Include planes whose normal is perpendicular to gravity. */
-	Vertical,
-
-	/** Include planes whose normal is parallel to gravity. */
-	Horizontal,
-
-	/** Include planes with arbitrary normals. */
-	Arbitrary,
-
-	/** If set, non-horizontal planes will be aligned perpendicular to gravity. */
-	OrientToGravity,
-
-	/** If set, inner planes will be returned; if not set, outer planes will be returned. */
-	PreferInner,
-
-	/** If set, holes in planar surfaces will be ignored. */
-	IgnoreHoles,
-
-	/** If set, include planes semantically tagged as ceiling. */
-	Ceiling,
-
-	/** If set, include planes semantically tagged as floor. */
-	Floor,
-
-	/** If set, include planes semantically tagged as wall. */
-	Wall
-};
\ No newline at end of file
