Removing some settings for debugging.

diff --git a/Engine/Plugins/Lumin/MagicLeap/Source/MagicLeapAR/MagicLeapAR.Build.cs b/Engine/Plugins/Lumin/MagicLeap/Source/MagicLeapAR/MagicLeapAR.Build.cs
--- a/Engine/Plugins/Lumin/MagicLeap/Source/MagicLeapAR/MagicLeapAR.Build.cs
+++ b/Engine/Plugins/Lumin/MagicLeap/Source/MagicLeapAR/MagicLeapAR.Build.cs
@@ -8,9 +8,6 @@ namespace UnrealBuildTool.Rules
 	{
 		public MagicLeapAR(ReadOnlyTargetRules Target) : base(Target)
 		{
-			//JB - Remove
-            OptimizeCode = CodeOptimization.Never;
-
             PrivateIncludePaths.AddRange(
 				new string[] {
 					"../../../../Source/Runtime/Renderer/Private",
