Added the update scene understanding AR support to the Magic Leap AR API

diff --git a/Engine/Plugins/Lumin/MagicLeap/Source/MagicLeapAR/Private/LuminARAPI.cpp b/Engine/Plugins/Lumin/MagicLeap/Source/MagicLeapAR/Private/LuminARAPI.cpp
--- a/Engine/Plugins/Lumin/MagicLeap/Source/MagicLeapAR/Private/LuminARAPI.cpp
+++ b/Engine/Plugins/Lumin/MagicLeap/Source/MagicLeapAR/Private/LuminARAPI.cpp
@@ -1,952 +1,989 @@
-// Copyright 1998-2019 Epic Games, Inc. All Rights Reserved.
-
-#include "LuminARAPI.h"
-
-#include "Misc/EngineVersion.h"
-#include "DrawDebugHelpers.h"
-#include "Templates/Casts.h"
-
-#include "IMagicLeapPlugin.h"
-#include "MagicLeapPrivileges.h"
-#include "MagicLeapHMD.h"
-#include "AppFramework.h"
-#include "HeadMountedDisplayFunctionLibrary.h"
-
-#if WITH_MLSDK
-#pragma warning( push )
-#pragma warning( disable : 4201)
-#include "ml_planes.h"
-#pragma warning( pop ) 
-#endif //WITH_MLSDK
-
-namespace
-{
-#if PLATFORM_LUMIN
-	static const FMatrix LuminARToUnrealTransform = FMatrix(
-		FPlane(0.0f, 0.0f, -1.0f, 0.0f),
-		FPlane(1.0f, 0.0f, 0.0f, 0.0f),
-		FPlane(0.0f, 1.0f, 0.0f, 0.0f),
-		FPlane(0.0f, 0.0f, 0.0f, 1.0f));
-
-	static const FMatrix LuminARToUnrealTransformInverse = LuminARToUnrealTransform.InverseFast();
-
-	FTransform LuminARPoseToUnrealTransform(const ArPose& InPose, const FLuminARSession* Session, float WorldToMeterScale)
-	{
-		FTransform LuminArPoseTransform;
-		LuminArPoseTransform.SetTranslation(InPose.Pos);
-		LuminArPoseTransform.SetRotation(InPose.Quat);
-
-		FMatrix LuminARPoseMatrix = LuminArPoseTransform.ToMatrixNoScale();
-		FTransform Result = FTransform(LuminARToUnrealTransform * LuminARPoseMatrix * LuminARToUnrealTransformInverse);
-		Result.SetLocation(Result.GetLocation() * WorldToMeterScale);
-
-		return Result;
-	}
-
-	void UnrealTransformToLuminARPose(const FTransform& UnrealTransform, const FLuminARSession* Session, ArPose& OutPose, float WorldToMeterScale)
-	{
-		FMatrix UnrealPoseMatrix = UnrealTransform.ToMatrixNoScale();
-		UnrealPoseMatrix.SetOrigin(UnrealPoseMatrix.GetOrigin() / WorldToMeterScale);
-		FMatrix LuminARPoseMatrix = LuminARToUnrealTransformInverse * UnrealPoseMatrix * LuminARToUnrealTransform;
-
-		FVector ArPosePosition = LuminARPoseMatrix.GetOrigin();
-		FQuat ArPoseRotation = LuminARPoseMatrix.ToQuat();
-
-		OutPose.Pos = LuminARPoseMatrix.GetOrigin();
-		OutPose.Quat = LuminARPoseMatrix.ToQuat();
-	}
-#endif
-	inline bool CheckIsSessionValid(FString TypeName, const TWeakPtr<FLuminARSession>& SessionPtr)
-	{
-		return true;
-	}
-}
-
-
-/****************************************/
-/*         FLuminARSession         */
-/****************************************/
-FLuminARSession::FLuminARSession()
-	: LatestFrame(nullptr)
-	, UObjectManager(nullptr)
-	, CachedWorldToMeterScale(100.0f)
-	, FrameNumber(0)
-
-{
-	// Create Lumin ARSession handle.
-	LatestFrame = new FLuminARFrame(this);
-#if PLATFORM_LUMIN
-	InitTracker();
-	LatestFrame->Init();
-#endif
-}
-
-FLuminARSession::~FLuminARSession()
-{
-	for (UARPin* Anchor : UObjectManager->AllAnchors)
-	{
-		Anchor->OnTrackingStateChanged(EARTrackingState::StoppedTracking);
-	}
-
-	delete LatestFrame;
-
-#if PLATFORM_LUMIN
-	DestroyTracker();
-#endif
-}
-
-
-ULuminARUObjectManager* FLuminARSession::GetUObjectManager()
-{
-	return UObjectManager;
-}
-
-float FLuminARSession::GetWorldToMeterScale()
-{
-	return CachedWorldToMeterScale;
-}
-
-
-#if PLATFORM_ANDROID
-MLHandle FLuminARSession::GetPlaneTrackerHandle()
-{
-	return PlaneTrackerHandle;
-}
-#endif
-
-ELuminARAPIStatus FLuminARSession::Resume()
-{
-	ELuminARAPIStatus ResumeStatus = ELuminARAPIStatus::AR_SUCCESS;
-#if PLATFORM_LUMIN
-	//InitTracker();
-#endif
-	return ResumeStatus;
-}
-
-ELuminARAPIStatus FLuminARSession::Pause()
-{
-	ELuminARAPIStatus PauseStatus = ELuminARAPIStatus::AR_SUCCESS;
-#if PLATFORM_LUMIN
-	//DestroyTracker();
-#endif
-
-	for (UARPin* Anchor : UObjectManager->AllAnchors)
-	{
-		Anchor->OnTrackingStateChanged(EARTrackingState::NotTracking);
-	}
-
-	return PauseStatus;
-}
-
-ELuminARAPIStatus FLuminARSession::Update(float WorldToMeterScale)
-{
-	ELuminARAPIStatus UpdateStatus = ELuminARAPIStatus::AR_SUCCESS;
-
-	CachedWorldToMeterScale = WorldToMeterScale;
-	int64 LastFrameTimestamp = LatestFrame->GetCameraTimestamp();
-	LatestFrame->Update(WorldToMeterScale);
-	if (LastFrameTimestamp != LatestFrame->GetCameraTimestamp())
-	{
-		FrameNumber++;
-	}
-
-	return UpdateStatus;
-}
-
-const FLuminARFrame* FLuminARSession::GetLatestFrame()
-{
-	return LatestFrame;
-}
-
-uint32 FLuminARSession::GetFrameNum()
-{
-	return FrameNumber;
-}
-
-
-// Anchors and Planes
-ELuminARAPIStatus FLuminARSession::CreateARAnchor(const FTransform& TransfromInTrackingSpace, UARTrackedGeometry* TrackedGeometry, USceneComponent* ComponentToPin, FName InDebugName, UARPin*& OutAnchor)
-{
-	ELuminARAPIStatus AnchorCreateStatus = ELuminARAPIStatus::AR_SUCCESS;
-	OutAnchor = nullptr;
-
-#if PLATFORM_LUMIN
-	TSharedPtr<LuminArAnchor> NewLuminArAnchor;
-	ArPose Pose;
-	UnrealTransformToLuminARPose(TransfromInTrackingSpace, this, Pose, CachedWorldToMeterScale);
-	if (TrackedGeometry == nullptr)
-	{
-		NewLuminArAnchor = MakeShared<LuminArAnchor>(Pose, ML_INVALID_HANDLE);
-	}
-	else
-	{
-		ensure(TrackedGeometry->GetNativeResource() != nullptr);
-		MLHandle ParentHandle = reinterpret_cast<FLuminARTrackableResource*>(TrackedGeometry->GetNativeResource())->GetNativeHandle();
-		ensure(ParentHandle != ML_INVALID_HANDLE);
-		NewLuminArAnchor = MakeShared<LuminArAnchor>(Pose, ParentHandle);
-	}
-
-	if (AnchorCreateStatus == ELuminARAPIStatus::AR_SUCCESS)
-	{
-		OutAnchor = NewObject<UARPin>();
-		OutAnchor->InitARPin(GetARSystem(), ComponentToPin, TransfromInTrackingSpace, TrackedGeometry, InDebugName);
-		OutAnchor->SetNativeResource(reinterpret_cast<void*>(NewLuminArAnchor.Get()));
-
-		UObjectManager->HandleToLuminAnchorMap.Add(NewLuminArAnchor->Handle, NewLuminArAnchor);
-		UObjectManager->AllAnchors.Add(OutAnchor);
-		UObjectManager->HandleToAnchorMap.Add(NewLuminArAnchor->Handle, OutAnchor);
-	}
-#endif
-	return AnchorCreateStatus;
-}
-
-void FLuminARSession::DetachAnchor(UARPin* Anchor)
-{
-	if (!UObjectManager->AllAnchors.Contains(Anchor))
-	{
-		return;
-	}
-
-#if PLATFORM_LUMIN
-	LuminArAnchor* NativeResource = reinterpret_cast<LuminArAnchor*>(Anchor->GetNativeResource());
-	check(NativeResource);
-	NativeResource->Detach();
-	Anchor->SetNativeResource(nullptr);
-
-	Anchor->OnTrackingStateChanged(EARTrackingState::StoppedTracking);
-
-	UObjectManager->HandleToAnchorMap.Remove(NativeResource->Handle);
-	UObjectManager->HandleToLuminAnchorMap.Remove(NativeResource->Handle);
-#endif
-	UObjectManager->AllAnchors.Remove(Anchor);
-}
-
-void FLuminARSession::GetAllAnchors(TArray<UARPin*>& OutAnchors) const
-{
-	OutAnchors = UObjectManager->AllAnchors;
-}
-
-void FLuminARSession::InitTracker()
-{
-#if PLATFORM_LUMIN
-	if (!MLHandleIsValid(PlaneTrackerHandle))
-	{
-		MLResult PlaneCreateResult = MLPlanesCreate(&PlaneTrackerHandle);
-		if (PlaneCreateResult != MLResult_Ok)
-		{
-			UE_LOG(LogLuminARAPI, Warning, TEXT("Failed to create Plane Tracker for Lumin AR Session Result:%i"), PlaneCreateResult);
-		}
-	}
-	else
-	{
-		UE_LOG(LogLuminARAPI, Warning, TEXT("Tracker already exists"));
-	}
-#endif
-}
-
-void FLuminARSession::DestroyTracker()
-{
-#if PLATFORM_LUMIN
-	if (MLHandleIsValid(PlaneTrackerHandle))
-	{
-		MLResult Result = MLPlanesDestroy(PlaneTrackerHandle);
-		if (Result != MLResult_Ok)
-		{
-			UE_LOG(LogLuminARAPI, Warning, TEXT("Failed to destroy Plane Tracker for Lumin AR Session Result:%i"), Result);
-		}
-
-		PlaneTrackerHandle = ML_INVALID_HANDLE;
-	}
-#endif //WITH_MLSDK
-}
-
-void FLuminARSession::AddReferencedObjects(FReferenceCollector& Collector)
-{
-	if (UObjectManager)
-	{
-		Collector.AddReferencedObject(UObjectManager);
-	}
-}
-
-/****************************************/
-/*         FLuminARFrame           */
-/****************************************/
-FLuminARFrame::FLuminARFrame(FLuminARSession* InSession)
-	: Session(InSession)
-	//, LatestCameraPose(FTransform::Identity)
-	, LatestCameraTimestamp(0)
-	, LatestCameraTrackingState(ELuminARTrackingState::StoppedTracking)
-	, LatestARPlaneQueryStatus(ELuminARPlaneQueryStatus::Unknown)
-#if PLATFORM_LUMIN
-	, PlaneTrackerHandle(ML_INVALID_HANDLE)
-#endif
-{
-}
-
-FLuminARFrame::~FLuminARFrame()
-{
-}
-
-void FLuminARFrame::Init()
-{
-#if PLATFORM_LUMIN
-	if (Session->GetPlaneTrackerHandle())
-	{
-		PlaneTrackerHandle = Session->GetPlaneTrackerHandle();
-}
-#endif
-}
-
-
-void FLuminARFrame::Update(float WorldToMeterScale)
-{
-#if PLATFORM_LUMIN
-	if (!MLHandleIsValid(PlaneTrackerHandle))
-	{
-		LatestCameraTrackingState = ELuminARTrackingState::NotTracking;
-		return;
-	}
-
-	// Update trackable that is cached in Unreal
-	StartPlaneQuery();
-	ProcessPlaneQuery();
-
-	switch (LatestARPlaneQueryStatus)
-	{
-	case ELuminARPlaneQueryStatus::Unknown:
-		return;
-	case ELuminARPlaneQueryStatus::Success:
-		{
-			//TODO - is pose needed?
-			//ArPose Pose;
-			//LatestCameraPose = Pose;
-
-			LatestCameraTimestamp = FPlatformTime::Seconds();
-			LatestCameraTrackingState = ELuminARTrackingState::Tracking;
-		}
-		break;
-	case ELuminARPlaneQueryStatus::Fail:
-		{
-			LatestCameraTrackingState = ELuminARTrackingState::NotTracking;
-		}
-		return;
-	default:
-		check(false);
-		return;
-	}
-
-	// Update Anchors
-	UpdatedAnchors.Empty();
-	for (auto HandleToAnchorMapPair : Session->GetUObjectManager()->HandleToAnchorMap)
-	{
-		const MLHandle AnchorHandle = HandleToAnchorMapPair.Key;
-		UARPin* const AnchorPin = HandleToAnchorMapPair.Value;
-		LuminArAnchor* LuminArAnchor = reinterpret_cast<struct LuminArAnchor*>(AnchorPin->GetNativeResource());
-		check(LuminArAnchor);
-		const MLHandle ParentTrackableHandle = LuminArAnchor->ParentTrackable;
-		if (ParentTrackableHandle != ML_INVALID_HANDLE)
-		{
-			UARTrackedGeometry* const ParentTrackable = Session->GetUObjectManager()->GetTrackableFromHandle<UARTrackedGeometry>(ParentTrackableHandle, Session);
-			const EARTrackingState AnchorTrackingState = ParentTrackable->GetTrackingState();
-			if (AnchorPin->GetTrackingState() != EARTrackingState::StoppedTracking)
-			{
-				AnchorPin->OnTrackingStateChanged(AnchorTrackingState);
-			}
-
-			if (AnchorPin->GetTrackingState() == EARTrackingState::Tracking)
-			{
-				AnchorPin->OnTransformUpdated(ParentTrackable->GetLocalToTrackingTransform());
-			}
-			UpdatedAnchors.Add(AnchorPin);
-		}
-	}
-#endif
-}
-
-//FTransform FLuminARFrame::GetCameraPose() const
-//{
-//	return LatestCameraPose;
-//}
-
-int64 FLuminARFrame::GetCameraTimestamp() const
-{
-	return LatestCameraTimestamp;
-}
-
-ELuminARTrackingState FLuminARFrame::GetCameraTrackingState() const
-{
-	return LatestCameraTrackingState;
-}
-
-void FLuminARFrame::GetUpdatedAnchors(TArray<UARPin*>& OutUpdatedAnchors) const
-{
-	OutUpdatedAnchors = UpdatedAnchors;
-}
-
-void FLuminARFrame::ARLineTrace(FVector2D ScreenPosition, ELuminARLineTraceChannel RequestedTraceChannels, TArray<FARTraceResult>& OutHitResults) const
-{
-#if PLATFORM_LUMIN
-	// Only testing straight forward from a little below the headset... Lumin isn't a handheld, but it's nice to have this do something.
-	IXRTrackingSystem* XRTrackingSystem = Session->GetARSystem()->GetXRTrackingSystem();
-	TArray<int32> Devices;
-	XRTrackingSystem->EnumerateTrackedDevices(Devices, EXRTrackedDeviceType::HeadMountedDisplay);
-	check(Devices.Num() == 1);
-	if (Devices.Num() > 0)
-	{
-		const int32 HMDDeviceID = Devices[0];
-		FQuat HMDQuat;
-		FVector HMDPosition;
-		const bool Success = XRTrackingSystem->GetCurrentPose(HMDDeviceID, HMDQuat, HMDPosition);
-		FTransform TrackingToWorldTransform = XRTrackingSystem->GetTrackingToWorldTransform();
-		if (Success)
-		{
-			const FVector HMDWorldPosition = TrackingToWorldTransform.TransformPosition(HMDPosition);
-			const FQuat HMDWorldQuat = TrackingToWorldTransform.TransformRotation(HMDQuat);
-			const FVector Start = HMDWorldPosition + FVector(0.0f, 0.0f, -10.0f);
-			const FVector Direction = HMDWorldQuat.Vector();
-			const FVector End = Start + (Direction * 10000.0f);
-
-			ARLineTrace(Start, End, RequestedTraceChannels, OutHitResults);
-		}
-	}
-#endif
-}
-void FLuminARFrame::ARLineTrace(FVector Start, FVector End, ELuminARLineTraceChannel RequestedTraceChannels, TArray<FARTraceResult>& OutHitResults) const
-{
-#if PLATFORM_LUMIN
-
-	// Only testing vs planes now, but not the ground plane.
-	ELuminARLineTraceChannel AllPlaneTraceChannels = /*ELuminARLineTraceChannel::InfinitePlane |*/ ELuminARLineTraceChannel::PlaneUsingExtent | ELuminARLineTraceChannel::PlaneUsingBoundaryPolygon;
-	if (!(RequestedTraceChannels & AllPlaneTraceChannels))
-		return;
-
-	TArray<UARPlaneGeometry*> Planes;
-	Session->GetAllTrackables(Planes);
-
-	for (UARPlaneGeometry* PPlane : Planes)
-	{
-		check(PPlane);
-		UARPlaneGeometry& Plane = *PPlane;
-
-		const FTransform LocalToWorld = Plane.GetLocalToWorldTransform();
-		const FVector PlaneOrigin = LocalToWorld.GetLocation();
-		const FVector PlaneNormal = LocalToWorld.TransformVectorNoScale(FVector(0, 0, 1));
-		const FVector Dir = End - Start;
-		// check if Dir is parallel to plane, no intersection
-		if (!FMath::IsNearlyZero(Dir | PlaneNormal, KINDA_SMALL_NUMBER))
-		{
-			// if T < 0 or > 1 we are outside the line segment, no intersection
-			float T = (((PlaneOrigin - Start) | PlaneNormal) / ((End - Start) | PlaneNormal));
-			if (T >= 0.0f || T <= 1.0f)
-			{
-				const FVector Intersection = Start + (Dir * T);
-
-				EARLineTraceChannels FoundInTraceChannel = EARLineTraceChannels::None;
-
-				if (!!(RequestedTraceChannels & (ELuminARLineTraceChannel::PlaneUsingExtent | ELuminARLineTraceChannel::PlaneUsingBoundaryPolygon)))
-				{
-					const FTransform WorldToLocal = LocalToWorld.Inverse();
-					const FVector LocalIntersection = WorldToLocal.TransformPosition(Intersection);
-
-					// Note: doing boundary check first for consistency with ARCore
-
-					if (FoundInTraceChannel == EARLineTraceChannels::None
-						&& !!(RequestedTraceChannels & ELuminARLineTraceChannel::PlaneUsingBoundaryPolygon))
-					{
-						// Note: could optimize this by computing an aligned boundary bounding rectangle during plane update and testing that first.
-
-						// Did we hit inside the boundary?
-						const TArray<FVector> Boundary = Plane.GetBoundaryPolygonInLocalSpace();
-						if (Boundary.Num() > 3) // has to be at least a triangle to have an inside
-						{
-							// This is the 'ray casting algorithm' for detecting if a point is inside a polygon.
-
-							// Draw a line from the point to the outside. Test for intersect with all edges.  If an odd number of edges are intersected the point is inside the polygon.
-							// The bounds are in plane local space, such that the plane is the x-y plane (z is always zero).
-							// We will offset that to put the point we are testing at 0,0 and our ray will be the +y axis (meaning the endpoint is 0,infinity and certainly outside the polygon).
-
-							// This could get the wrong answer if the test line goes exactly through a boundary vertex because that would register as two intersections.
-							// We are ignoring this rare failure cases.
-
-							const FVector2D Origin(LocalIntersection.X, LocalIntersection.Y);
-							const int Num = Boundary.Num();
-							FVector2D A(Boundary[Num - 1].X - Origin.X, Boundary[Num - 1].Y - Origin.Y);
-							int32 Crossings = 0;
-							for (int i = 0; i < Num; ++i)
-							{
-								const FVector2D B(Boundary[i].X - Origin.X, Boundary[i].Y - Origin.Y);
-
-								// Check if there is any Y intercept in the line segment.
-								if (FMath::Sign(A.X) != FMath::Sign(B.X))
-								{
-									// Check if the Y intercept is positive.
-									const float Slope = (B.Y - A.Y) / (B.X - A.X);
-									const float YIntercept = A.Y - (Slope * A.X);
-									if (YIntercept > 0.0f)
-									{
-										Crossings += 1;
-									}
-								}
-
-								A = B;
-							}
-							if ((Crossings & 0x01) == 0x01)
-							{
-								FoundInTraceChannel = EARLineTraceChannels::PlaneUsingBoundaryPolygon;
-							}
-						}
-					}
-
-					if (FoundInTraceChannel == EARLineTraceChannels::None
-						&& !!(RequestedTraceChannels & ELuminARLineTraceChannel::PlaneUsingExtent))
-					{
-						// Did we hit inside the plane extents?
-						if (FMath::Abs(LocalIntersection.X) <= Plane.GetExtent().X
-							&& FMath::Abs(LocalIntersection.Y) <= Plane.GetExtent().Y)
-						{
-							FoundInTraceChannel = EARLineTraceChannels::PlaneUsingExtent;
-						}
-					}
-				}
-
-				//// This 'infinite plane' 'ground plane' stuff seems... weird.
-				//if (FoundInTraceChannel == EARLineTraceChannels::None
-				//	&&!!(RequestedTraceChannels & ELuminARLineTraceChannel::InfinitePlane))
-				//{
-				//	FoundInTraceChannel = EARLineTraceChannels::GroundPlane;
-				//}
-
-				// write the result
-				if (FoundInTraceChannel != EARLineTraceChannels::None)
-				{
-					const float Distance = Dir.Size() * T;
-
-					FTransform HitTransform = LocalToWorld;
-					HitTransform.SetLocation(Intersection);
-
-					FARTraceResult UEHitResult(Session->GetARSystem(), Distance, FoundInTraceChannel, HitTransform, &Plane);
-					UEHitResult.SetLocalToWorldTransform(HitTransform);
-					OutHitResults.Add(UEHitResult);
-				}
-			}
-		}
-	}
-
-	// Sort closest to furthest
-	OutHitResults.Sort(FARTraceResult::FARTraceResultComparer());
-#endif
-}
-
-
-FMatrix FLuminARFrame::GetProjectionMatrix() const
-{
-	FMatrix ProjectionMatrix;
-
-#if PLATFORM_LUMIN
-	if (Session == nullptr)
-	{
-		return ProjectionMatrix;
-	}
-
-	//TODO - get projection from session
-
-	// Unreal uses the infinite far plane project matrix.
-	ProjectionMatrix.M[2][2] = 0.0f;
-	ProjectionMatrix.M[2][3] = 1.0f;
-	ProjectionMatrix.M[3][2] = GNearClippingPlane;
-#endif
-	return ProjectionMatrix;
-}
-
-void FLuminARFrame::TransformDisplayUvCoords(const TArray<float>& UvCoords, TArray<float>& OutUvCoords) const
-{
-#if PLATFORM_LUMIN
-	OutUvCoords = UvCoords;
-#endif
-}
-
-
-FLuminARLightEstimate FLuminARFrame::GetLightEstimate() const
-{
-#if PLATFORM_LUMIN
-	//TODO - Fill in light estimate.  See ml_lighting_tracking.h.
-	FLuminARLightEstimate LightEstimate;
-	return LightEstimate;
-#else
-	return FLuminARLightEstimate();
-#endif
-}
-
-void FLuminARFrame::StartPlaneQuery()
-{
-#if PLATFORM_LUMIN
-	//if we haven't queried yet, start one!
-	if (PlaneQueryHandle == ML_INVALID_HANDLE)
-	{
-		if (IMagicLeapPlugin::Get().IsMagicLeapHMDValid())
-		{
-			const FAppFramework& AppFramework = static_cast<FMagicLeapHMD*>(GEngine->XRSystem->GetHMDDevice())->GetAppFrameworkConst();
-			float WorldToMetersScale = AppFramework.GetWorldToMetersScale();
-			check(WorldToMetersScale != 0);
-
-			FTransform PoseInverse = UHeadMountedDisplayFunctionLibrary::GetTrackingToWorldTransform(nullptr).Inverse();
-			FPlane plane;
-
-			// Apply Lumin specific AR session config, if available.  Otherwise use default values.
-			MaxPlaneQueryResults = 200;
-			int32 MinPlaneArea = 25;
-			TArray<EPlaneQueryFlags> QueryFlags;
-			FVector SearchVolumeExtents(10000.0f, 10000.0f, 10000.0f);
-			bDiscardZeroExtentPlanes = false;
-
-			const UARSessionConfig& ARSessionConfig = Session->GetARSystem()->AccessSessionConfig();
-			const ULuminARSessionConfig* LuminARSessionConfig = Cast<ULuminARSessionConfig>(&ARSessionConfig);
-			if (LuminARSessionConfig != nullptr)
-			{
-				MaxPlaneQueryResults = LuminARSessionConfig->MaxPlaneQueryResults;
-				MinPlaneArea = LuminARSessionConfig->MinPlaneArea;
-				if (LuminARSessionConfig->ShouldDoHorizontalPlaneDetection())	{ QueryFlags.Add(EPlaneQueryFlags::Horizontal); }
-				if (LuminARSessionConfig->ShouldDoVerticalPlaneDetection())		{ QueryFlags.Add(EPlaneQueryFlags::Vertical); }
-				if (LuminARSessionConfig->bArbitraryOrientationPlaneDetection)	{ QueryFlags.Add(EPlaneQueryFlags::Arbitrary); }
-				SearchVolumeExtents = LuminARSessionConfig->PlaneSearchExtents;
-				for (EPlaneQueryFlags flag : LuminARSessionConfig->PlaneQueryFlags)
-				{
-					QueryFlags.Add(flag);
-				}
-				bDiscardZeroExtentPlanes = LuminARSessionConfig->bDiscardZeroExtentPlanes;
-			}
-			else
-			{
-				UE_LOG(LogLuminARAPI, Log, TEXT("LuminArSessionConfig not found, using defaults for lumin specific settings."));
-				QueryFlags.Add(EPlaneQueryFlags::Vertical);
-				QueryFlags.Add(EPlaneQueryFlags::Horizontal);
-			}
-
-			MLPlanesQuery Query;
-			Query.max_results = static_cast<uint32>(MaxPlaneQueryResults);
-			Query.flags = UnrealToMLPlanesQueryFlags(QueryFlags) | MLPlanesQueryFlag_Polygons;
-			Query.min_hole_length = 50.0f / WorldToMetersScale;  // ML docs say this value is deprecated, so presumably it does nothing now.
-			Query.min_plane_area = MinPlaneArea / (WorldToMetersScale * WorldToMetersScale);
-			Query.bounds_center = MagicLeap::ToMLVector(PoseInverse.GetTranslation(), WorldToMetersScale);
-			Query.bounds_rotation = MagicLeap::ToMLQuat(PoseInverse.GetRotation());
-			Query.bounds_extents = MagicLeap::ToMLVector(SearchVolumeExtents, WorldToMetersScale);
-
-			// MagicLeap::ToMLVector() causes the Z component to be negated.
-			// The bounds were thus invalid and resulted in everything being tracked. 
-			// This provides the content devs with an option to ignore the bounding volume at will.
-			{
-				Query.bounds_extents.x = FMath::Abs<float>(Query.bounds_extents.x);
-				Query.bounds_extents.y = FMath::Abs<float>(Query.bounds_extents.y);
-				Query.bounds_extents.z = FMath::Abs<float>(Query.bounds_extents.z);
-			}
-
-			MLResult QueryResult = MLPlanesQueryBegin(PlaneTrackerHandle, &Query, &PlaneQueryHandle);
-			if (QueryResult != MLResult_Ok || !MLHandleIsValid(PlaneQueryHandle))
-			{
-				UE_LOG(LogLuminARAPI, Error, TEXT("LuminARFrame could not request planes."));
-			}
-		}
-	}
-#endif //PLATFORM_LUMIN
-}
-
-void FLuminARFrame::ProcessPlaneQuery()
-{
-#if PLATFORM_LUMIN
-	if (PlaneQueryHandle != ML_INVALID_HANDLE)
-	{
-		FTransform PoseTransform = FTransform::Identity; 
-
-		uint32 OutNumResults = 0;
-		TArray<MLPlane> ResultMLPlanes;
-		ResultMLPlanes.AddDefaulted(MaxPlaneQueryResults);
-
-		MLPlaneBoundariesList PlaneBoundariesList;
-		MLPlaneBoundariesListInit(&PlaneBoundariesList);
-
-		MLResult PlaneQueryResult = MLPlanesQueryGetResultsWithBoundaries(PlaneTrackerHandle, PlaneQueryHandle, ResultMLPlanes.GetData(), &OutNumResults, &PlaneBoundariesList);
-		switch (PlaneQueryResult)
-		{
-		case MLResult_Pending:
-			// Intentionally skip. We'll continue to check until it has completed.
-			break;
-		case MLResult_UnspecifiedFailure:
-		{
-			UE_LOG(LogLuminARAPI, Error, TEXT("MLPlanesQueryGetResults MLResult_UnspecifiedFailure."));
-			PlaneQueryHandle = ML_INVALID_HANDLE;
-			LatestARPlaneQueryStatus = ELuminARPlaneQueryStatus::Fail;
-			break;
-		}
-		case MLResult_Ok:
-		{
-			const FAppFramework& AppFramework = static_cast<FMagicLeapHMD*>(GEngine->XRSystem->GetHMDDevice())->GetAppFrameworkConst();
-			float WorldToMetersScale = AppFramework.GetWorldToMetersScale();
-
-			PlaneResultsMap.Reset();
-			PlaneResultsMap.Reserve(OutNumResults);
-
-			// Setup for boundaries, build a map of Handles to Boundaries
-			TMap<uint64, const MLPlaneBoundaries*> HandleToBoundariesMap;
-			{
-				const uint32 BoundaryCount = PlaneBoundariesList.plane_boundaries_count;
-				for (uint32 b = 0; b < BoundaryCount; ++b)
-				{
-					const MLPlaneBoundaries& Boundaries = PlaneBoundariesList.plane_boundaries[b];
-					HandleToBoundariesMap.Add(Boundaries.id, &Boundaries);
-				}
-			}
-			const FRotator RotateToContentOrientation(-90.0f, 0.0f, 0.0f);
-			const FTransform RotateToContentTransform(RotateToContentOrientation);
-
-			for (uint32 i = 0; i < OutNumResults; ++i)
-			{
-				const MLPlane& ResultMLPlane = ResultMLPlanes[i];
-				FPlaneResult ResultUEPlane;
-
-				const uint64_t Mask = TNumericLimits<uint32>::Max();
-				ResultUEPlane.ID.A = ResultMLPlane.id & Mask;
-				ResultUEPlane.ID.B = ResultMLPlane.id >> 32;
-				ResultUEPlane.ID_64 = ResultMLPlane.id;
-
-				// Perception uses all coordinates in RUB so for them X axis is right and corresponds to the width of the plane.
-				// Unreal uses FRU, so the Y-axis is towards the right which makes the Y component of the vector the width.
-				ResultUEPlane.PlaneDimensions = FVector2D(ResultMLPlane.height * WorldToMetersScale, ResultMLPlane.width * WorldToMetersScale);
-				if (ResultUEPlane.PlaneDimensions.X == 0.0f || ResultUEPlane.PlaneDimensions.Y == 0.0f)
-				{
-					if (bDiscardZeroExtentPlanes)
-					{
-						continue;
-					}
-				}
-
-				FTransform planeTransform = FTransform(MagicLeap::ToFQuat(ResultMLPlane.rotation), MagicLeap::ToFVector(ResultMLPlane.position, WorldToMetersScale), FVector(1.0f, 1.0f, 1.0f));
-				if (planeTransform.ContainsNaN())
-				{
-					continue;
-				}
-				if (!planeTransform.GetRotation().IsNormalized())
-				{
-					FQuat rotation = planeTransform.GetRotation();
-					rotation.Normalize();
-					planeTransform.SetRotation(rotation);
-				}
-
-				planeTransform.SetRotation(MagicLeap::ToUERotator(planeTransform.GetRotation()));
-				planeTransform.AddToTranslation(PoseTransform.GetLocation());
-				planeTransform.ConcatenateRotation(PoseTransform.Rotator().Quaternion());
-				ResultUEPlane.PlanePosition = planeTransform.GetLocation();
-				ResultUEPlane.PlaneOrientation = planeTransform.Rotator();
-				// The plane orientation has the forward axis (X) pointing in the direction of the plane's normal.
-				// We are rotating it by 90 degrees clock-wise about the right axis (Y) to get the up vector (Z) to point in the direction of the plane's normal.
-				// ResultPlane we are rotating the axis, the rotation is in the opposite direction of the object i.e. -90 degrees.
-				ResultUEPlane.ContentOrientation = UKismetMathLibrary::Conv_VectorToRotator(UKismetMathLibrary::RotateAngleAxis(UKismetMathLibrary::Conv_RotatorToVector(ResultUEPlane.PlaneOrientation), -90, UKismetMathLibrary::GetRightVector(ResultUEPlane.PlaneOrientation)));
-				MLToUnrealPlanesQueryFlags(ResultMLPlane.flags, ResultUEPlane.PlaneFlags);
-
-				// Boundaries
-				const MLPlaneBoundaries** const BoundariesPtr = HandleToBoundariesMap.Find(ResultMLPlane.id);
-				const MLPlaneBoundaries* const Boundaries = BoundariesPtr ? *BoundariesPtr : nullptr;
-				if (Boundaries && Boundaries->boundaries_count > 0)
-				{
-					FTransform PlaneTransformInverse = planeTransform.Inverse();
-					FTransform BoundaryVertTransform = PlaneTransformInverse * RotateToContentTransform;
-
-					// Seems like there is really only one boundary polygon...
-					const int BoundariesCount = Boundaries->boundaries_count;
-					for (int b = 0; b < BoundariesCount; ++b)
-					{
-						const MLPlaneBoundary& Boundary = Boundaries->boundaries[b];
-						const MLPolygon* Polygon = Boundary.polygon;
-						check(Polygon);
-						const uint32 VertCount = Polygon->vertices_count;
-						for (uint32 v = 0; v < VertCount; ++v)
-						{
-							MLVec3f& Vert = Polygon->vertices[v];
-							FVector LocalVert = BoundaryVertTransform.TransformPosition(MagicLeap::ToFVector(Vert, WorldToMetersScale));
-							ResultUEPlane.BoundaryPolygon.Add(LocalVert);
-						}
-					}
-				}
-
-				PlaneResultsMap.Add(ResultMLPlane.id, ResultUEPlane);
-			}
-
-			PlaneQueryHandle = ML_INVALID_HANDLE;
-
-			// Mark planes that previously existed, but no longer do StoppedTracking.
-			const auto & TrackableHandleMap = Session->GetUObjectManager()->TrackableHandleMap;
-			for (auto TrackablePair : TrackableHandleMap)
-			{
-				uint64 Handle = TrackablePair.Key;
-				if (!PlaneResultsMap.Contains(Handle))
-				{
-					UARTrackedGeometry* ARTrackedGeometry = TrackablePair.Value.Get();
-					if (ARTrackedGeometry)
-					{
-						ARTrackedGeometry->SetTrackingState(EARTrackingState::StoppedTracking);
-					}
-				}
-			}
-
-			for (auto PlaneResultPair : PlaneResultsMap)
-			{
-				UARTrackedGeometry* ARTrackedGeometry = Session->GetUObjectManager()->GetTrackableFromHandle<UARPlaneGeometry>(PlaneResultPair.Key, Session);
-
-				if (ARTrackedGeometry && ARTrackedGeometry->GetTrackingState() != EARTrackingState::StoppedTracking)
-				{
-					ARTrackedGeometry->SetTrackingState(EARTrackingState::Tracking);
-					FLuminARTrackableResource* TrackableResource = reinterpret_cast<FLuminARTrackableResource*>(ARTrackedGeometry->GetNativeResource());
-					TrackableResource->UpdateGeometryData(Session);
-				}
-			}
-			LatestARPlaneQueryStatus = ELuminARPlaneQueryStatus::Success;
-			break;
-		}
-		default:
-			UE_LOG(LogLuminARAPI, Warning, TEXT("Unexpected return code from MLPlanesQueryGetResults: %d"), PlaneQueryResult);
-			LatestARPlaneQueryStatus = ELuminARPlaneQueryStatus::Fail;
-		}
-		MLPlanesReleaseBoundariesList(PlaneTrackerHandle, &PlaneBoundariesList);
-	}
-#endif //PLATFORM_LUMIN
-}
-
-
-
-TSharedPtr<FLuminARSession> FLuminARSession::CreateLuminARSession()
-{
-	TSharedPtr<FLuminARSession> NewSession = MakeShared<FLuminARSession>();
-
-	ULuminARUObjectManager* UObjectManager = NewObject<ULuminARUObjectManager>();
-	UObjectManager->AddToRoot();
-
-	NewSession->UObjectManager = UObjectManager;
-	return NewSession;
-}
-
-
-/************************************************/
-/*       ULuminARTrackableResource         */
-/************************************************/
-#if PLATFORM_LUMIN
-
-EARTrackingState FLuminARTrackableResource::GetTrackingState()
-{
-	if (MLHandleIsValid(TrackableHandle))
-	{
-		check(TrackedGeometry);
-		return TrackedGeometry->GetTrackingState();
-	}
-
-	return EARTrackingState::NotTracking;
-}
-
-void FLuminARTrackableResource::UpdateGeometryData(FLuminARSession* InSession)
-{
-	TrackedGeometry->UpdateTrackingState(GetTrackingState());
-}
-
-void FLuminARTrackableResource::ResetNativeHandle(LuminArTrackable* InTrackableHandle)
-{
-	TrackableHandle = ML_INVALID_HANDLE;
-	if (InTrackableHandle != nullptr)
-	{
-		TrackableHandle = InTrackableHandle->Handle;
-	}
-
-	UpdateGeometryData(nullptr);
-}
-
-void FLuminARTrackedPlaneResource::UpdateGeometryData(FLuminARSession* InSession)
-{
-	FLuminARTrackableResource::UpdateGeometryData(InSession);
-
-	UARPlaneGeometry* PlaneGeometry = CastChecked<UARPlaneGeometry>(TrackedGeometry);
-
-	if (!InSession || /*!CheckIsSessionValid("LuminARPlane", Session) || */TrackedGeometry->GetTrackingState() == EARTrackingState::StoppedTracking)
-	{
-		return;
-	}
-
-	const FLuminARFrame* Frame = InSession->GetLatestFrame();
-	if (Frame == nullptr)
-	{
-		return;
-	}
-
-	// PlaneResult is in unreal tracking space, so already scaled and axis corrected
-	const FPlaneResult* PlaneResult = Frame->GetPlaneResult(TrackableHandle);
-	check(PlaneResult);
-
-	FTransform LocalToTrackingTransform(PlaneResult->ContentOrientation, PlaneResult->PlanePosition);
-	FVector Extent(PlaneResult->PlaneDimensions.X * 0.5f, PlaneResult->PlaneDimensions.Y * 0.5f, 0);  // Extent is half the width and height
-
-	uint32 FrameNum = InSession->GetFrameNum();
-	int64 TimeStamp = InSession->GetLatestFrame()->GetCameraTimestamp();
-
-	UARPlaneGeometry* SubsumedByPlane = nullptr; // ARCore only
-	PlaneGeometry->UpdateTrackedGeometry(InSession->GetARSystem(), FrameNum, static_cast<double>(TimeStamp), LocalToTrackingTransform, InSession->GetARSystem()->GetAlignmentTransform(), FVector::ZeroVector, Extent, PlaneResult->BoundaryPolygon, SubsumedByPlane);
-	PlaneGeometry->SetDebugName(FName(TEXT("LuminARPlane")));
-}
-#endif
-
-
-#if PLATFORM_LUMIN
-void ULuminARUObjectManager::DumpTrackableHandleMap(const MLHandle SessionHandle)
-{
-	UE_LOG(LogLuminARAPI, Log, TEXT("ULuminARUObjectManager::DumpTrackableHandleMap"));
-	for (auto KeyValuePair : TrackableHandleMap)
-	{
-		uint64 TrackableHandle = KeyValuePair.Key;
-		TWeakObjectPtr<UARTrackedGeometry> TrackedGeometry = KeyValuePair.Value;
-		UE_LOG(LogLuminARAPI, Log, TEXT("  Trackable Handle %llu"), TrackableHandle);
-		if (TrackedGeometry.IsValid())
-		{
-			UARTrackedGeometry* TrackedGeometryObj = TrackedGeometry.Get();
-			FLuminARTrackableResource* NativeResource = reinterpret_cast<FLuminARTrackableResource*>(TrackedGeometryObj->GetNativeResource());
-			UE_LOG(LogLuminARAPI, Log, TEXT("  TrackedGeometry - NativeResource:%p, type: %s, tracking state: %d"),
-				NativeResource->GetNativeHandle(), *TrackedGeometryObj->GetClass()->GetFName().ToString(), (int)TrackedGeometryObj->GetTrackingState());
-		}
-		else
-		{
-			UE_LOG(LogLuminARAPI, Log, TEXT("  TrackedGeometry - InValid or Pending Kill."))
-		}
-	}
-}
-#endif
-
-ELuminARAPIStatus FLuminARSession::AcquireCameraImage(ULuminARCameraImage *&OutCameraImage)
-{
-	ELuminARAPIStatus ApiStatus = ELuminARAPIStatus::AR_SUCCESS;
-#if PLATFORM_LUMIN
-	if (LatestFrame == nullptr)
-	{
-		return ELuminARAPIStatus::AR_ERROR_FATAL;
-	}
-
-	ArImage *OutImage = nullptr;
-	//TODO - support acquiring camera
-
-	if (ApiStatus == ELuminARAPIStatus::AR_SUCCESS)
-	{
-		UE_LOG(LogLuminARAPI, Error, TEXT("Support camera image"));
-	}
-	else
-	{
-		UE_LOG(LogLuminARAPI, Error, TEXT("AcquireCameraImage failed!"));
-	}
-#endif
-
-	return ApiStatus;
-}
-
-void* FLuminARSession::GetLatestFrameRawPointer()
-{
-#if PLATFORM_LUMIN
-	//TODO - do we need this?
-#endif
-	return nullptr;
-}
-
-
-
-
+// Copyright 1998-2019 Epic Games, Inc. All Rights Reserved.
+
+#include "LuminARAPI.h"
+
+#include "Misc/EngineVersion.h"
+#include "DrawDebugHelpers.h"
+#include "Templates/Casts.h"
+
+#include "IMagicLeapPlugin.h"
+#include "MagicLeapPrivileges.h"
+#include "MagicLeapHMD.h"
+#include "AppFramework.h"
+#include "HeadMountedDisplayFunctionLibrary.h"
+
+#if WITH_MLSDK
+#pragma warning( push )
+#pragma warning( disable : 4201)
+#include "ml_planes.h"
+#pragma warning( pop ) 
+#endif //WITH_MLSDK
+
+namespace
+{
+#if PLATFORM_LUMIN
+	static const FMatrix LuminARToUnrealTransform = FMatrix(
+		FPlane(0.0f, 0.0f, -1.0f, 0.0f),
+		FPlane(1.0f, 0.0f, 0.0f, 0.0f),
+		FPlane(0.0f, 1.0f, 0.0f, 0.0f),
+		FPlane(0.0f, 0.0f, 0.0f, 1.0f));
+
+	static const FMatrix LuminARToUnrealTransformInverse = LuminARToUnrealTransform.InverseFast();
+
+	FTransform LuminARPoseToUnrealTransform(const ArPose& InPose, const FLuminARSession* Session, float WorldToMeterScale)
+	{
+		FTransform LuminArPoseTransform;
+		LuminArPoseTransform.SetTranslation(InPose.Pos);
+		LuminArPoseTransform.SetRotation(InPose.Quat);
+
+		FMatrix LuminARPoseMatrix = LuminArPoseTransform.ToMatrixNoScale();
+		FTransform Result = FTransform(LuminARToUnrealTransform * LuminARPoseMatrix * LuminARToUnrealTransformInverse);
+		Result.SetLocation(Result.GetLocation() * WorldToMeterScale);
+
+		return Result;
+	}
+
+	void UnrealTransformToLuminARPose(const FTransform& UnrealTransform, const FLuminARSession* Session, ArPose& OutPose, float WorldToMeterScale)
+	{
+		FMatrix UnrealPoseMatrix = UnrealTransform.ToMatrixNoScale();
+		UnrealPoseMatrix.SetOrigin(UnrealPoseMatrix.GetOrigin() / WorldToMeterScale);
+		FMatrix LuminARPoseMatrix = LuminARToUnrealTransformInverse * UnrealPoseMatrix * LuminARToUnrealTransform;
+
+		FVector ArPosePosition = LuminARPoseMatrix.GetOrigin();
+		FQuat ArPoseRotation = LuminARPoseMatrix.ToQuat();
+
+		OutPose.Pos = LuminARPoseMatrix.GetOrigin();
+		OutPose.Quat = LuminARPoseMatrix.ToQuat();
+	}
+#endif
+	inline bool CheckIsSessionValid(FString TypeName, const TWeakPtr<FLuminARSession>& SessionPtr)
+	{
+		return true;
+	}
+}
+
+
+/****************************************/
+/*         FLuminARSession         */
+/****************************************/
+FLuminARSession::FLuminARSession()
+	: LatestFrame(nullptr)
+	, UObjectManager(nullptr)
+	, CachedWorldToMeterScale(100.0f)
+	, FrameNumber(0)
+
+{
+	// Create Lumin ARSession handle.
+	LatestFrame = new FLuminARFrame(this);
+#if PLATFORM_LUMIN
+	InitTracker();
+	LatestFrame->Init();
+#endif
+}
+
+FLuminARSession::~FLuminARSession()
+{
+	for (UARPin* Anchor : UObjectManager->AllAnchors)
+	{
+		Anchor->OnTrackingStateChanged(EARTrackingState::StoppedTracking);
+	}
+
+	delete LatestFrame;
+
+#if PLATFORM_LUMIN
+	DestroyTracker();
+#endif
+}
+
+
+ULuminARUObjectManager* FLuminARSession::GetUObjectManager()
+{
+	return UObjectManager;
+}
+
+float FLuminARSession::GetWorldToMeterScale()
+{
+	return CachedWorldToMeterScale;
+}
+
+
+#if PLATFORM_ANDROID
+MLHandle FLuminARSession::GetPlaneTrackerHandle()
+{
+	return PlaneTrackerHandle;
+}
+#endif
+
+ELuminARAPIStatus FLuminARSession::Resume()
+{
+	ELuminARAPIStatus ResumeStatus = ELuminARAPIStatus::AR_SUCCESS;
+#if PLATFORM_LUMIN
+	//InitTracker();
+#endif
+	return ResumeStatus;
+}
+
+ELuminARAPIStatus FLuminARSession::Pause()
+{
+	ELuminARAPIStatus PauseStatus = ELuminARAPIStatus::AR_SUCCESS;
+#if PLATFORM_LUMIN
+	//DestroyTracker();
+#endif
+
+	for (UARPin* Anchor : UObjectManager->AllAnchors)
+	{
+		Anchor->OnTrackingStateChanged(EARTrackingState::NotTracking);
+	}
+
+	return PauseStatus;
+}
+
+ELuminARAPIStatus FLuminARSession::Update(float WorldToMeterScale)
+{
+	ELuminARAPIStatus UpdateStatus = ELuminARAPIStatus::AR_SUCCESS;
+
+	CachedWorldToMeterScale = WorldToMeterScale;
+	int64 LastFrameTimestamp = LatestFrame->GetCameraTimestamp();
+	LatestFrame->Update(WorldToMeterScale);
+	if (LastFrameTimestamp != LatestFrame->GetCameraTimestamp())
+	{
+		FrameNumber++;
+	}
+
+	return UpdateStatus;
+}
+
+const FLuminARFrame* FLuminARSession::GetLatestFrame()
+{
+	return LatestFrame;
+}
+
+uint32 FLuminARSession::GetFrameNum()
+{
+	return FrameNumber;
+}
+
+
+// Anchors and Planes
+ELuminARAPIStatus FLuminARSession::CreateARAnchor(const FTransform& TransfromInTrackingSpace, UARTrackedGeometry* TrackedGeometry, USceneComponent* ComponentToPin, FName InDebugName, UARPin*& OutAnchor)
+{
+	ELuminARAPIStatus AnchorCreateStatus = ELuminARAPIStatus::AR_SUCCESS;
+	OutAnchor = nullptr;
+
+#if PLATFORM_LUMIN
+	TSharedPtr<LuminArAnchor> NewLuminArAnchor;
+	ArPose Pose;
+	UnrealTransformToLuminARPose(TransfromInTrackingSpace, this, Pose, CachedWorldToMeterScale);
+	if (TrackedGeometry == nullptr)
+	{
+		NewLuminArAnchor = MakeShared<LuminArAnchor>(Pose, ML_INVALID_HANDLE);
+	}
+	else
+	{
+		ensure(TrackedGeometry->GetNativeResource() != nullptr);
+		MLHandle ParentHandle = reinterpret_cast<FLuminARTrackableResource*>(TrackedGeometry->GetNativeResource())->GetNativeHandle();
+		ensure(ParentHandle != ML_INVALID_HANDLE);
+		NewLuminArAnchor = MakeShared<LuminArAnchor>(Pose, ParentHandle);
+	}
+
+	if (AnchorCreateStatus == ELuminARAPIStatus::AR_SUCCESS)
+	{
+		OutAnchor = NewObject<UARPin>();
+		OutAnchor->InitARPin(GetARSystem(), ComponentToPin, TransfromInTrackingSpace, TrackedGeometry, InDebugName);
+		OutAnchor->SetNativeResource(reinterpret_cast<void*>(NewLuminArAnchor.Get()));
+
+		UObjectManager->HandleToLuminAnchorMap.Add(NewLuminArAnchor->Handle, NewLuminArAnchor);
+		UObjectManager->AllAnchors.Add(OutAnchor);
+		UObjectManager->HandleToAnchorMap.Add(NewLuminArAnchor->Handle, OutAnchor);
+	}
+#endif
+	return AnchorCreateStatus;
+}
+
+void FLuminARSession::DetachAnchor(UARPin* Anchor)
+{
+	if (!UObjectManager->AllAnchors.Contains(Anchor))
+	{
+		return;
+	}
+
+#if PLATFORM_LUMIN
+	LuminArAnchor* NativeResource = reinterpret_cast<LuminArAnchor*>(Anchor->GetNativeResource());
+	check(NativeResource);
+	NativeResource->Detach();
+	Anchor->SetNativeResource(nullptr);
+
+	Anchor->OnTrackingStateChanged(EARTrackingState::StoppedTracking);
+
+	UObjectManager->HandleToAnchorMap.Remove(NativeResource->Handle);
+	UObjectManager->HandleToLuminAnchorMap.Remove(NativeResource->Handle);
+#endif
+	UObjectManager->AllAnchors.Remove(Anchor);
+}
+
+void FLuminARSession::GetAllAnchors(TArray<UARPin*>& OutAnchors) const
+{
+	OutAnchors = UObjectManager->AllAnchors;
+}
+
+void FLuminARSession::InitTracker()
+{
+#if PLATFORM_LUMIN
+	if (!MLHandleIsValid(PlaneTrackerHandle))
+	{
+		MLResult PlaneCreateResult = MLPlanesCreate(&PlaneTrackerHandle);
+		if (PlaneCreateResult != MLResult_Ok)
+		{
+			UE_LOG(LogLuminARAPI, Warning, TEXT("Failed to create Plane Tracker for Lumin AR Session Result:%i"), PlaneCreateResult);
+		}
+	}
+	else
+	{
+		UE_LOG(LogLuminARAPI, Warning, TEXT("Tracker already exists"));
+	}
+#endif
+}
+
+void FLuminARSession::DestroyTracker()
+{
+#if PLATFORM_LUMIN
+	if (MLHandleIsValid(PlaneTrackerHandle))
+	{
+		MLResult Result = MLPlanesDestroy(PlaneTrackerHandle);
+		if (Result != MLResult_Ok)
+		{
+			UE_LOG(LogLuminARAPI, Warning, TEXT("Failed to destroy Plane Tracker for Lumin AR Session Result:%i"), Result);
+		}
+
+		PlaneTrackerHandle = ML_INVALID_HANDLE;
+	}
+#endif //WITH_MLSDK
+}
+
+void FLuminARSession::AddReferencedObjects(FReferenceCollector& Collector)
+{
+	if (UObjectManager)
+	{
+		Collector.AddReferencedObject(UObjectManager);
+	}
+}
+
+/****************************************/
+/*         FLuminARFrame           */
+/****************************************/
+FLuminARFrame::FLuminARFrame(FLuminARSession* InSession)
+	: Session(InSession)
+	//, LatestCameraPose(FTransform::Identity)
+	, LatestCameraTimestamp(0)
+	, LatestCameraTrackingState(ELuminARTrackingState::StoppedTracking)
+	, LatestARPlaneQueryStatus(ELuminARPlaneQueryStatus::Unknown)
+#if PLATFORM_LUMIN
+	, PlaneTrackerHandle(ML_INVALID_HANDLE)
+#endif
+{
+}
+
+FLuminARFrame::~FLuminARFrame()
+{
+}
+
+void FLuminARFrame::Init()
+{
+#if PLATFORM_LUMIN
+	if (Session->GetPlaneTrackerHandle())
+	{
+		PlaneTrackerHandle = Session->GetPlaneTrackerHandle();
+}
+#endif
+}
+
+
+void FLuminARFrame::Update(float WorldToMeterScale)
+{
+#if PLATFORM_LUMIN
+	if (!MLHandleIsValid(PlaneTrackerHandle))
+	{
+		LatestCameraTrackingState = ELuminARTrackingState::NotTracking;
+		return;
+	}
+
+	// Update trackable that is cached in Unreal
+	StartPlaneQuery();
+	ProcessPlaneQuery();
+
+	switch (LatestARPlaneQueryStatus)
+	{
+	case ELuminARPlaneQueryStatus::Unknown:
+		return;
+	case ELuminARPlaneQueryStatus::Success:
+		{
+			//TODO - is pose needed?
+			//ArPose Pose;
+			//LatestCameraPose = Pose;
+
+			LatestCameraTimestamp = FPlatformTime::Seconds();
+			LatestCameraTrackingState = ELuminARTrackingState::Tracking;
+		}
+		break;
+	case ELuminARPlaneQueryStatus::Fail:
+		{
+			LatestCameraTrackingState = ELuminARTrackingState::NotTracking;
+		}
+		return;
+	default:
+		check(false);
+		return;
+	}
+
+	// Update Anchors
+	UpdatedAnchors.Empty();
+	for (auto HandleToAnchorMapPair : Session->GetUObjectManager()->HandleToAnchorMap)
+	{
+		const MLHandle AnchorHandle = HandleToAnchorMapPair.Key;
+		UARPin* const AnchorPin = HandleToAnchorMapPair.Value;
+		LuminArAnchor* LuminArAnchor = reinterpret_cast<struct LuminArAnchor*>(AnchorPin->GetNativeResource());
+		check(LuminArAnchor);
+		const MLHandle ParentTrackableHandle = LuminArAnchor->ParentTrackable;
+		if (ParentTrackableHandle != ML_INVALID_HANDLE)
+		{
+			UARTrackedGeometry* const ParentTrackable = Session->GetUObjectManager()->GetTrackableFromHandle<UARTrackedGeometry>(ParentTrackableHandle, Session);
+			const EARTrackingState AnchorTrackingState = ParentTrackable->GetTrackingState();
+			if (AnchorPin->GetTrackingState() != EARTrackingState::StoppedTracking)
+			{
+				AnchorPin->OnTrackingStateChanged(AnchorTrackingState);
+			}
+
+			if (AnchorPin->GetTrackingState() == EARTrackingState::Tracking)
+			{
+				AnchorPin->OnTransformUpdated(ParentTrackable->GetLocalToTrackingTransform());
+			}
+			UpdatedAnchors.Add(AnchorPin);
+		}
+	}
+#endif
+}
+
+//FTransform FLuminARFrame::GetCameraPose() const
+//{
+//	return LatestCameraPose;
+//}
+
+int64 FLuminARFrame::GetCameraTimestamp() const
+{
+	return LatestCameraTimestamp;
+}
+
+ELuminARTrackingState FLuminARFrame::GetCameraTrackingState() const
+{
+	return LatestCameraTrackingState;
+}
+
+void FLuminARFrame::GetUpdatedAnchors(TArray<UARPin*>& OutUpdatedAnchors) const
+{
+	OutUpdatedAnchors = UpdatedAnchors;
+}
+
+void FLuminARFrame::ARLineTrace(FVector2D ScreenPosition, ELuminARLineTraceChannel RequestedTraceChannels, TArray<FARTraceResult>& OutHitResults) const
+{
+#if PLATFORM_LUMIN
+	// Only testing straight forward from a little below the headset... Lumin isn't a handheld, but it's nice to have this do something.
+	IXRTrackingSystem* XRTrackingSystem = Session->GetARSystem()->GetXRTrackingSystem();
+	TArray<int32> Devices;
+	XRTrackingSystem->EnumerateTrackedDevices(Devices, EXRTrackedDeviceType::HeadMountedDisplay);
+	check(Devices.Num() == 1);
+	if (Devices.Num() > 0)
+	{
+		const int32 HMDDeviceID = Devices[0];
+		FQuat HMDQuat;
+		FVector HMDPosition;
+		const bool Success = XRTrackingSystem->GetCurrentPose(HMDDeviceID, HMDQuat, HMDPosition);
+		FTransform TrackingToWorldTransform = XRTrackingSystem->GetTrackingToWorldTransform();
+		if (Success)
+		{
+			const FVector HMDWorldPosition = TrackingToWorldTransform.TransformPosition(HMDPosition);
+			const FQuat HMDWorldQuat = TrackingToWorldTransform.TransformRotation(HMDQuat);
+			const FVector Start = HMDWorldPosition + FVector(0.0f, 0.0f, -10.0f);
+			const FVector Direction = HMDWorldQuat.Vector();
+			const FVector End = Start + (Direction * 10000.0f);
+
+			ARLineTrace(Start, End, RequestedTraceChannels, OutHitResults);
+		}
+	}
+#endif
+}
+void FLuminARFrame::ARLineTrace(FVector Start, FVector End, ELuminARLineTraceChannel RequestedTraceChannels, TArray<FARTraceResult>& OutHitResults) const
+{
+#if PLATFORM_LUMIN
+
+	// Only testing vs planes now, but not the ground plane.
+	ELuminARLineTraceChannel AllPlaneTraceChannels = /*ELuminARLineTraceChannel::InfinitePlane |*/ ELuminARLineTraceChannel::PlaneUsingExtent | ELuminARLineTraceChannel::PlaneUsingBoundaryPolygon;
+	if (!(RequestedTraceChannels & AllPlaneTraceChannels))
+		return;
+
+	TArray<UARPlaneGeometry*> Planes;
+	Session->GetAllTrackables(Planes);
+
+	for (UARPlaneGeometry* PPlane : Planes)
+	{
+		check(PPlane);
+		UARPlaneGeometry& Plane = *PPlane;
+
+		const FTransform LocalToWorld = Plane.GetLocalToWorldTransform();
+		const FVector PlaneOrigin = LocalToWorld.GetLocation();
+		const FVector PlaneNormal = LocalToWorld.TransformVectorNoScale(FVector(0, 0, 1));
+		const FVector Dir = End - Start;
+		// check if Dir is parallel to plane, no intersection
+		if (!FMath::IsNearlyZero(Dir | PlaneNormal, KINDA_SMALL_NUMBER))
+		{
+			// if T < 0 or > 1 we are outside the line segment, no intersection
+			float T = (((PlaneOrigin - Start) | PlaneNormal) / ((End - Start) | PlaneNormal));
+			if (T >= 0.0f || T <= 1.0f)
+			{
+				const FVector Intersection = Start + (Dir * T);
+
+				EARLineTraceChannels FoundInTraceChannel = EARLineTraceChannels::None;
+
+				if (!!(RequestedTraceChannels & (ELuminARLineTraceChannel::PlaneUsingExtent | ELuminARLineTraceChannel::PlaneUsingBoundaryPolygon)))
+				{
+					const FTransform WorldToLocal = LocalToWorld.Inverse();
+					const FVector LocalIntersection = WorldToLocal.TransformPosition(Intersection);
+
+					// Note: doing boundary check first for consistency with ARCore
+
+					if (FoundInTraceChannel == EARLineTraceChannels::None
+						&& !!(RequestedTraceChannels & ELuminARLineTraceChannel::PlaneUsingBoundaryPolygon))
+					{
+						// Note: could optimize this by computing an aligned boundary bounding rectangle during plane update and testing that first.
+
+						// Did we hit inside the boundary?
+						const TArray<FVector> Boundary = Plane.GetBoundaryPolygonInLocalSpace();
+						if (Boundary.Num() > 3) // has to be at least a triangle to have an inside
+						{
+							// This is the 'ray casting algorithm' for detecting if a point is inside a polygon.
+
+							// Draw a line from the point to the outside. Test for intersect with all edges.  If an odd number of edges are intersected the point is inside the polygon.
+							// The bounds are in plane local space, such that the plane is the x-y plane (z is always zero).
+							// We will offset that to put the point we are testing at 0,0 and our ray will be the +y axis (meaning the endpoint is 0,infinity and certainly outside the polygon).
+
+							// This could get the wrong answer if the test line goes exactly through a boundary vertex because that would register as two intersections.
+							// We are ignoring this rare failure cases.
+
+							const FVector2D Origin(LocalIntersection.X, LocalIntersection.Y);
+							const int Num = Boundary.Num();
+							FVector2D A(Boundary[Num - 1].X - Origin.X, Boundary[Num - 1].Y - Origin.Y);
+							int32 Crossings = 0;
+							for (int i = 0; i < Num; ++i)
+							{
+								const FVector2D B(Boundary[i].X - Origin.X, Boundary[i].Y - Origin.Y);
+
+								// Check if there is any Y intercept in the line segment.
+								if (FMath::Sign(A.X) != FMath::Sign(B.X))
+								{
+									// Check if the Y intercept is positive.
+									const float Slope = (B.Y - A.Y) / (B.X - A.X);
+									const float YIntercept = A.Y - (Slope * A.X);
+									if (YIntercept > 0.0f)
+									{
+										Crossings += 1;
+									}
+								}
+
+								A = B;
+							}
+							if ((Crossings & 0x01) == 0x01)
+							{
+								FoundInTraceChannel = EARLineTraceChannels::PlaneUsingBoundaryPolygon;
+							}
+						}
+					}
+
+					if (FoundInTraceChannel == EARLineTraceChannels::None
+						&& !!(RequestedTraceChannels & ELuminARLineTraceChannel::PlaneUsingExtent))
+					{
+						// Did we hit inside the plane extents?
+						if (FMath::Abs(LocalIntersection.X) <= Plane.GetExtent().X
+							&& FMath::Abs(LocalIntersection.Y) <= Plane.GetExtent().Y)
+						{
+							FoundInTraceChannel = EARLineTraceChannels::PlaneUsingExtent;
+						}
+					}
+				}
+
+				//// This 'infinite plane' 'ground plane' stuff seems... weird.
+				//if (FoundInTraceChannel == EARLineTraceChannels::None
+				//	&&!!(RequestedTraceChannels & ELuminARLineTraceChannel::InfinitePlane))
+				//{
+				//	FoundInTraceChannel = EARLineTraceChannels::GroundPlane;
+				//}
+
+				// write the result
+				if (FoundInTraceChannel != EARLineTraceChannels::None)
+				{
+					const float Distance = Dir.Size() * T;
+
+					FTransform HitTransform = LocalToWorld;
+					HitTransform.SetLocation(Intersection);
+
+					FARTraceResult UEHitResult(Session->GetARSystem(), Distance, FoundInTraceChannel, HitTransform, &Plane);
+					UEHitResult.SetLocalToWorldTransform(HitTransform);
+					OutHitResults.Add(UEHitResult);
+				}
+			}
+		}
+	}
+
+	// Sort closest to furthest
+	OutHitResults.Sort(FARTraceResult::FARTraceResultComparer());
+#endif
+}
+
+
+FMatrix FLuminARFrame::GetProjectionMatrix() const
+{
+	FMatrix ProjectionMatrix;
+
+#if PLATFORM_LUMIN
+	if (Session == nullptr)
+	{
+		return ProjectionMatrix;
+	}
+
+	//TODO - get projection from session
+
+	// Unreal uses the infinite far plane project matrix.
+	ProjectionMatrix.M[2][2] = 0.0f;
+	ProjectionMatrix.M[2][3] = 1.0f;
+	ProjectionMatrix.M[3][2] = GNearClippingPlane;
+#endif
+	return ProjectionMatrix;
+}
+
+void FLuminARFrame::TransformDisplayUvCoords(const TArray<float>& UvCoords, TArray<float>& OutUvCoords) const
+{
+#if PLATFORM_LUMIN
+	OutUvCoords = UvCoords;
+#endif
+}
+
+
+FLuminARLightEstimate FLuminARFrame::GetLightEstimate() const
+{
+#if PLATFORM_LUMIN
+	//TODO - Fill in light estimate.  See ml_lighting_tracking.h.
+	FLuminARLightEstimate LightEstimate;
+	return LightEstimate;
+#else
+	return FLuminARLightEstimate();
+#endif
+}
+
+void FLuminARFrame::StartPlaneQuery()
+{
+#if PLATFORM_LUMIN
+	//if we haven't queried yet, start one!
+	if (PlaneQueryHandle == ML_INVALID_HANDLE)
+	{
+		if (IMagicLeapPlugin::Get().IsMagicLeapHMDValid())
+		{
+			const FAppFramework& AppFramework = static_cast<FMagicLeapHMD*>(GEngine->XRSystem->GetHMDDevice())->GetAppFrameworkConst();
+			float WorldToMetersScale = AppFramework.GetWorldToMetersScale();
+			check(WorldToMetersScale != 0);
+
+			FTransform PoseInverse = UHeadMountedDisplayFunctionLibrary::GetTrackingToWorldTransform(nullptr).Inverse();
+			FPlane plane;
+
+			// Apply Lumin specific AR session config, if available.  Otherwise use default values.
+			MaxPlaneQueryResults = 200;
+			int32 MinPlaneArea = 25;
+			TArray<EPlaneQueryFlags> QueryFlags;
+			FVector SearchVolumeExtents(10000.0f, 10000.0f, 10000.0f);
+			bDiscardZeroExtentPlanes = false;
+
+			const UARSessionConfig& ARSessionConfig = Session->GetARSystem()->AccessSessionConfig();
+			const ULuminARSessionConfig* LuminARSessionConfig = Cast<ULuminARSessionConfig>(&ARSessionConfig);
+			if (LuminARSessionConfig != nullptr)
+			{
+				MaxPlaneQueryResults = LuminARSessionConfig->MaxPlaneQueryResults;
+				MinPlaneArea = LuminARSessionConfig->MinPlaneArea;
+				if (LuminARSessionConfig->ShouldDoHorizontalPlaneDetection())	{ QueryFlags.Add(EPlaneQueryFlags::Horizontal); }
+				if (LuminARSessionConfig->ShouldDoVerticalPlaneDetection())		{ QueryFlags.Add(EPlaneQueryFlags::Vertical); }
+				if (LuminARSessionConfig->bArbitraryOrientationPlaneDetection)	{ QueryFlags.Add(EPlaneQueryFlags::Arbitrary); }
+				SearchVolumeExtents = LuminARSessionConfig->PlaneSearchExtents;
+				for (EPlaneQueryFlags flag : LuminARSessionConfig->PlaneQueryFlags)
+				{
+					QueryFlags.Add(flag);
+				}
+				bDiscardZeroExtentPlanes = LuminARSessionConfig->bDiscardZeroExtentPlanes;
+			}
+			else
+			{
+				UE_LOG(LogLuminARAPI, Log, TEXT("LuminArSessionConfig not found, using defaults for lumin specific settings."));
+				QueryFlags.Add(EPlaneQueryFlags::Vertical);
+				QueryFlags.Add(EPlaneQueryFlags::Horizontal);
+			}
+
+			MLPlanesQuery Query;
+			Query.max_results = static_cast<uint32>(MaxPlaneQueryResults);
+			Query.flags = UnrealToMLPlanesQueryFlags(QueryFlags) | MLPlanesQueryFlag_Polygons;
+			Query.min_hole_length = 50.0f / WorldToMetersScale;  // ML docs say this value is deprecated, so presumably it does nothing now.
+			Query.min_plane_area = MinPlaneArea / (WorldToMetersScale * WorldToMetersScale);
+			Query.bounds_center = MagicLeap::ToMLVector(PoseInverse.GetTranslation(), WorldToMetersScale);
+			Query.bounds_rotation = MagicLeap::ToMLQuat(PoseInverse.GetRotation());
+			Query.bounds_extents = MagicLeap::ToMLVector(SearchVolumeExtents, WorldToMetersScale);
+
+			// MagicLeap::ToMLVector() causes the Z component to be negated.
+			// The bounds were thus invalid and resulted in everything being tracked. 
+			// This provides the content devs with an option to ignore the bounding volume at will.
+			{
+				Query.bounds_extents.x = FMath::Abs<float>(Query.bounds_extents.x);
+				Query.bounds_extents.y = FMath::Abs<float>(Query.bounds_extents.y);
+				Query.bounds_extents.z = FMath::Abs<float>(Query.bounds_extents.z);
+			}
+
+			MLResult QueryResult = MLPlanesQueryBegin(PlaneTrackerHandle, &Query, &PlaneQueryHandle);
+			if (QueryResult != MLResult_Ok || !MLHandleIsValid(PlaneQueryHandle))
+			{
+				UE_LOG(LogLuminARAPI, Error, TEXT("LuminARFrame could not request planes."));
+			}
+		}
+	}
+#endif //PLATFORM_LUMIN
+}
+
+void FLuminARFrame::ProcessPlaneQuery()
+{
+#if PLATFORM_LUMIN
+	if (PlaneQueryHandle != ML_INVALID_HANDLE)
+	{
+		FTransform PoseTransform = FTransform::Identity; 
+
+		uint32 OutNumResults = 0;
+		TArray<MLPlane> ResultMLPlanes;
+		ResultMLPlanes.AddDefaulted(MaxPlaneQueryResults);
+
+		MLPlaneBoundariesList PlaneBoundariesList;
+		MLPlaneBoundariesListInit(&PlaneBoundariesList);
+
+		MLResult PlaneQueryResult = MLPlanesQueryGetResultsWithBoundaries(PlaneTrackerHandle, PlaneQueryHandle, ResultMLPlanes.GetData(), &OutNumResults, &PlaneBoundariesList);
+		switch (PlaneQueryResult)
+		{
+		case MLResult_Pending:
+			// Intentionally skip. We'll continue to check until it has completed.
+			break;
+		case MLResult_UnspecifiedFailure:
+		{
+			UE_LOG(LogLuminARAPI, Error, TEXT("MLPlanesQueryGetResults MLResult_UnspecifiedFailure."));
+			PlaneQueryHandle = ML_INVALID_HANDLE;
+			LatestARPlaneQueryStatus = ELuminARPlaneQueryStatus::Fail;
+			break;
+		}
+		case MLResult_Ok:
+		{
+			const FAppFramework& AppFramework = static_cast<FMagicLeapHMD*>(GEngine->XRSystem->GetHMDDevice())->GetAppFrameworkConst();
+			float WorldToMetersScale = AppFramework.GetWorldToMetersScale();
+
+			PlaneResultsMap.Reset();
+			PlaneResultsMap.Reserve(OutNumResults);
+
+			// Setup for boundaries, build a map of Handles to Boundaries
+			TMap<uint64, const MLPlaneBoundaries*> HandleToBoundariesMap;
+			{
+				const uint32 BoundaryCount = PlaneBoundariesList.plane_boundaries_count;
+				for (uint32 b = 0; b < BoundaryCount; ++b)
+				{
+					const MLPlaneBoundaries& Boundaries = PlaneBoundariesList.plane_boundaries[b];
+					HandleToBoundariesMap.Add(Boundaries.id, &Boundaries);
+				}
+			}
+			const FRotator RotateToContentOrientation(-90.0f, 0.0f, 0.0f);
+			const FTransform RotateToContentTransform(RotateToContentOrientation);
+
+			for (uint32 i = 0; i < OutNumResults; ++i)
+			{
+				const MLPlane& ResultMLPlane = ResultMLPlanes[i];
+				FPlaneResult ResultUEPlane;
+
+				const uint64_t Mask = TNumericLimits<uint32>::Max();
+				ResultUEPlane.ID.A = ResultMLPlane.id & Mask;
+				ResultUEPlane.ID.B = ResultMLPlane.id >> 32;
+				ResultUEPlane.ID_64 = ResultMLPlane.id;
+
+				// Perception uses all coordinates in RUB so for them X axis is right and corresponds to the width of the plane.
+				// Unreal uses FRU, so the Y-axis is towards the right which makes the Y component of the vector the width.
+				ResultUEPlane.PlaneDimensions = FVector2D(ResultMLPlane.height * WorldToMetersScale, ResultMLPlane.width * WorldToMetersScale);
+				if (ResultUEPlane.PlaneDimensions.X == 0.0f || ResultUEPlane.PlaneDimensions.Y == 0.0f)
+				{
+					if (bDiscardZeroExtentPlanes)
+					{
+						continue;
+					}
+				}
+
+				FTransform planeTransform = FTransform(MagicLeap::ToFQuat(ResultMLPlane.rotation), MagicLeap::ToFVector(ResultMLPlane.position, WorldToMetersScale), FVector(1.0f, 1.0f, 1.0f));
+				if (planeTransform.ContainsNaN())
+				{
+					continue;
+				}
+				if (!planeTransform.GetRotation().IsNormalized())
+				{
+					FQuat rotation = planeTransform.GetRotation();
+					rotation.Normalize();
+					planeTransform.SetRotation(rotation);
+				}
+
+				planeTransform.SetRotation(MagicLeap::ToUERotator(planeTransform.GetRotation()));
+				planeTransform.AddToTranslation(PoseTransform.GetLocation());
+				planeTransform.ConcatenateRotation(PoseTransform.Rotator().Quaternion());
+				ResultUEPlane.PlanePosition = planeTransform.GetLocation();
+				ResultUEPlane.PlaneOrientation = planeTransform.Rotator();
+				// The plane orientation has the forward axis (X) pointing in the direction of the plane's normal.
+				// We are rotating it by 90 degrees clock-wise about the right axis (Y) to get the up vector (Z) to point in the direction of the plane's normal.
+				// ResultPlane we are rotating the axis, the rotation is in the opposite direction of the object i.e. -90 degrees.
+				ResultUEPlane.ContentOrientation = UKismetMathLibrary::Conv_VectorToRotator(UKismetMathLibrary::RotateAngleAxis(UKismetMathLibrary::Conv_RotatorToVector(ResultUEPlane.PlaneOrientation), -90, UKismetMathLibrary::GetRightVector(ResultUEPlane.PlaneOrientation)));
+				MLToUnrealPlanesQueryFlags(ResultMLPlane.flags, ResultUEPlane.PlaneFlags);
+
+				// Boundaries
+				const MLPlaneBoundaries** const BoundariesPtr = HandleToBoundariesMap.Find(ResultMLPlane.id);
+				const MLPlaneBoundaries* const Boundaries = BoundariesPtr ? *BoundariesPtr : nullptr;
+				if (Boundaries && Boundaries->boundaries_count > 0)
+				{
+					FTransform PlaneTransformInverse = planeTransform.Inverse();
+					FTransform BoundaryVertTransform = PlaneTransformInverse * RotateToContentTransform;
+
+					// Seems like there is really only one boundary polygon...
+					const int BoundariesCount = Boundaries->boundaries_count;
+					for (int b = 0; b < BoundariesCount; ++b)
+					{
+						const MLPlaneBoundary& Boundary = Boundaries->boundaries[b];
+						const MLPolygon* Polygon = Boundary.polygon;
+						check(Polygon);
+						const uint32 VertCount = Polygon->vertices_count;
+						for (uint32 v = 0; v < VertCount; ++v)
+						{
+							MLVec3f& Vert = Polygon->vertices[v];
+							FVector LocalVert = BoundaryVertTransform.TransformPosition(MagicLeap::ToFVector(Vert, WorldToMetersScale));
+							ResultUEPlane.BoundaryPolygon.Add(LocalVert);
+						}
+					}
+				}
+
+				PlaneResultsMap.Add(ResultMLPlane.id, ResultUEPlane);
+			}
+
+			PlaneQueryHandle = ML_INVALID_HANDLE;
+
+			// Mark planes that previously existed, but no longer do StoppedTracking.
+			const auto & TrackableHandleMap = Session->GetUObjectManager()->TrackableHandleMap;
+			for (auto TrackablePair : TrackableHandleMap)
+			{
+				uint64 Handle = TrackablePair.Key;
+				if (!PlaneResultsMap.Contains(Handle))
+				{
+					UARTrackedGeometry* ARTrackedGeometry = TrackablePair.Value.Get();
+					if (ARTrackedGeometry)
+					{
+						ARTrackedGeometry->SetTrackingState(EARTrackingState::StoppedTracking);
+					}
+				}
+			}
+
+			for (auto PlaneResultPair : PlaneResultsMap)
+			{
+				UARTrackedGeometry* ARTrackedGeometry = Session->GetUObjectManager()->GetTrackableFromHandle<UARPlaneGeometry>(PlaneResultPair.Key, Session);
+
+				if (ARTrackedGeometry && ARTrackedGeometry->GetTrackingState() != EARTrackingState::StoppedTracking)
+				{
+					ARTrackedGeometry->SetTrackingState(EARTrackingState::Tracking);
+					FLuminARTrackableResource* TrackableResource = reinterpret_cast<FLuminARTrackableResource*>(ARTrackedGeometry->GetNativeResource());
+					TrackableResource->UpdateGeometryData(Session);
+				}
+			}
+			LatestARPlaneQueryStatus = ELuminARPlaneQueryStatus::Success;
+			break;
+		}
+		default:
+			UE_LOG(LogLuminARAPI, Warning, TEXT("Unexpected return code from MLPlanesQueryGetResults: %d"), PlaneQueryResult);
+			LatestARPlaneQueryStatus = ELuminARPlaneQueryStatus::Fail;
+		}
+		MLPlanesReleaseBoundariesList(PlaneTrackerHandle, &PlaneBoundariesList);
+	}
+#endif //PLATFORM_LUMIN
+}
+
+
+
+TSharedPtr<FLuminARSession> FLuminARSession::CreateLuminARSession()
+{
+	TSharedPtr<FLuminARSession> NewSession = MakeShared<FLuminARSession>();
+
+	ULuminARUObjectManager* UObjectManager = NewObject<ULuminARUObjectManager>();
+	UObjectManager->AddToRoot();
+
+	NewSession->UObjectManager = UObjectManager;
+	return NewSession;
+}
+
+
+/************************************************/
+/*       ULuminARTrackableResource         */
+/************************************************/
+#if PLATFORM_LUMIN
+
+EARTrackingState FLuminARTrackableResource::GetTrackingState()
+{
+	if (MLHandleIsValid(TrackableHandle))
+	{
+		check(TrackedGeometry);
+		return TrackedGeometry->GetTrackingState();
+	}
+
+	return EARTrackingState::NotTracking;
+}
+
+void FLuminARTrackableResource::UpdateGeometryData(FLuminARSession* InSession)
+{
+	TrackedGeometry->UpdateTrackingState(GetTrackingState());
+}
+
+void FLuminARTrackableResource::ResetNativeHandle(LuminArTrackable* InTrackableHandle)
+{
+	TrackableHandle = ML_INVALID_HANDLE;
+	if (InTrackableHandle != nullptr)
+	{
+		TrackableHandle = InTrackableHandle->Handle;
+	}
+
+	UpdateGeometryData(nullptr);
+}
+
+void FLuminARTrackedPlaneResource::UpdateGeometryData(FLuminARSession* InSession)
+{
+	FLuminARTrackableResource::UpdateGeometryData(InSession);
+
+	UARPlaneGeometry* PlaneGeometry = CastChecked<UARPlaneGeometry>(TrackedGeometry);
+
+	if (!InSession || /*!CheckIsSessionValid("LuminARPlane", Session) || */TrackedGeometry->GetTrackingState() == EARTrackingState::StoppedTracking)
+	{
+		return;
+	}
+
+	const FLuminARFrame* Frame = InSession->GetLatestFrame();
+	if (Frame == nullptr)
+	{
+		return;
+	}
+
+	// PlaneResult is in unreal tracking space, so already scaled and axis corrected
+	const FPlaneResult* PlaneResult = Frame->GetPlaneResult(TrackableHandle);
+	check(PlaneResult);
+
+	FTransform LocalToTrackingTransform(PlaneResult->ContentOrientation, PlaneResult->PlanePosition);
+	FVector Extent(PlaneResult->PlaneDimensions.X * 0.5f, PlaneResult->PlaneDimensions.Y * 0.5f, 0);  // Extent is half the width and height
+
+	uint32 FrameNum = InSession->GetFrameNum();
+	int64 TimeStamp = InSession->GetLatestFrame()->GetCameraTimestamp();
+
+	UARPlaneGeometry* SubsumedByPlane = nullptr; // ARCore only
+	PlaneGeometry->UpdateTrackedGeometry(InSession->GetARSystem(), FrameNum, static_cast<double>(TimeStamp), LocalToTrackingTransform, InSession->GetARSystem()->GetAlignmentTransform(), FVector::ZeroVector, Extent, PlaneResult->BoundaryPolygon, SubsumedByPlane);
+	PlaneGeometry->SetDebugName(FName(TEXT("LuminARPlane")));
+	// Inspect the plane flags and set the scene understanding data based upon them
+	for (EPlaneQueryFlags Flag : PlaneResult->PlaneFlags)
+	{
+		switch (Flag)
+		{
+			case EPlaneQueryFlags::Vertical:
+			{
+				PlaneGeometry->SetOrientation(EARPlaneOrientation::Vertical);
+				break;
+			}
+			case EPlaneQueryFlags::Horizontal:
+			{
+				PlaneGeometry->SetOrientation(EARPlaneOrientation::Horizontal);
+				break;
+			}
+			case EPlaneQueryFlags::Arbitrary:
+			{
+				PlaneGeometry->SetOrientation(EARPlaneOrientation::Diagonal);
+				break;
+			}
+			case EPlaneQueryFlags::Ceiling:
+			{
+				PlaneGeometry->SetObjectClassification(EARObjectClassification::Ceiling);
+				break;
+			}
+			case EPlaneQueryFlags::Floor:
+			{
+				PlaneGeometry->SetObjectClassification(EARObjectClassification::Floor);
+				break;
+			}
+			case EPlaneQueryFlags::Wall:
+			{
+				PlaneGeometry->SetObjectClassification(EARObjectClassification::Wall);
+				break;
+			}
+		}
+	}
+}
+#endif
+
+
+#if PLATFORM_LUMIN
+void ULuminARUObjectManager::DumpTrackableHandleMap(const MLHandle SessionHandle)
+{
+	UE_LOG(LogLuminARAPI, Log, TEXT("ULuminARUObjectManager::DumpTrackableHandleMap"));
+	for (auto KeyValuePair : TrackableHandleMap)
+	{
+		uint64 TrackableHandle = KeyValuePair.Key;
+		TWeakObjectPtr<UARTrackedGeometry> TrackedGeometry = KeyValuePair.Value;
+		UE_LOG(LogLuminARAPI, Log, TEXT("  Trackable Handle %llu"), TrackableHandle);
+		if (TrackedGeometry.IsValid())
+		{
+			UARTrackedGeometry* TrackedGeometryObj = TrackedGeometry.Get();
+			FLuminARTrackableResource* NativeResource = reinterpret_cast<FLuminARTrackableResource*>(TrackedGeometryObj->GetNativeResource());
+			UE_LOG(LogLuminARAPI, Log, TEXT("  TrackedGeometry - NativeResource:%p, type: %s, tracking state: %d"),
+				NativeResource->GetNativeHandle(), *TrackedGeometryObj->GetClass()->GetFName().ToString(), (int)TrackedGeometryObj->GetTrackingState());
+		}
+		else
+		{
+			UE_LOG(LogLuminARAPI, Log, TEXT("  TrackedGeometry - InValid or Pending Kill."))
+		}
+	}
+}
+#endif
+
+ELuminARAPIStatus FLuminARSession::AcquireCameraImage(ULuminARCameraImage *&OutCameraImage)
+{
+	ELuminARAPIStatus ApiStatus = ELuminARAPIStatus::AR_SUCCESS;
+#if PLATFORM_LUMIN
+	if (LatestFrame == nullptr)
+	{
+		return ELuminARAPIStatus::AR_ERROR_FATAL;
+	}
+
+	ArImage *OutImage = nullptr;
+	//TODO - support acquiring camera
+
+	if (ApiStatus == ELuminARAPIStatus::AR_SUCCESS)
+	{
+		UE_LOG(LogLuminARAPI, Error, TEXT("Support camera image"));
+	}
+	else
+	{
+		UE_LOG(LogLuminARAPI, Error, TEXT("AcquireCameraImage failed!"));
+	}
+#endif
+
+	return ApiStatus;
+}
+
+void* FLuminARSession::GetLatestFrameRawPointer()
+{
+#if PLATFORM_LUMIN
+	//TODO - do we need this?
+#endif
+	return nullptr;
+}
+
+
+
+
