Remove redundant TSharedFromThis inheritance
Fixes shared pointer asserts on Lumin

diff --git a/Engine/Plugins/Lumin/MagicLeap/Source/MagicLeap/Private/MagicLeapHMD.h b/Engine/Plugins/Lumin/MagicLeap/Source/MagicLeap/Private/MagicLeapHMD.h
--- a/Engine/Plugins/Lumin/MagicLeap/Source/MagicLeap/Private/MagicLeapHMD.h
+++ b/Engine/Plugins/Lumin/MagicLeap/Source/MagicLeap/Private/MagicLeapHMD.h
@@ -44,7 +44,7 @@ public:
 /**
   * MagicLeap Head Mounted Display
   */
-class MAGICLEAP_API FMagicLeapHMD : public IMagicLeapHMD, public FHeadMountedDisplayBase, public FXRRenderTargetManager, public TSharedFromThis<FMagicLeapHMD, ESPMode::ThreadSafe>
+class MAGICLEAP_API FMagicLeapHMD : public IMagicLeapHMD, public FHeadMountedDisplayBase, public FXRRenderTargetManager
 {
 public:
 
