Rollback //UE4/Dev-VR/Engine/Plugins/Lumin/MagicLeap/Source/MagicLeap/Private/MagicLeapCustomPresent.cpp to revision 16
-Reverting to fix build break.

diff --git a/Engine/Plugins/Lumin/MagicLeap/Source/MagicLeap/Private/MagicLeapCustomPresent.cpp b/Engine/Plugins/Lumin/MagicLeap/Source/MagicLeap/Private/MagicLeapCustomPresent.cpp
--- a/Engine/Plugins/Lumin/MagicLeap/Source/MagicLeap/Private/MagicLeapCustomPresent.cpp
+++ b/Engine/Plugins/Lumin/MagicLeap/Source/MagicLeap/Private/MagicLeapCustomPresent.cpp
@@ -383,7 +383,6 @@ void FMagicLeapCustomPresentOpenGL::FinishRendering()
 			UE_LOG(LogMagicLeap, Error, TEXT("MLGraphicsEndFrame failed with status %d"), Result);
 		}
 	}
-	Plugin->GetCurrentFrame().bBeginFrameSucceeded = false;
 	Plugin->InitializeOldFrameFromRenderFrame();
 #endif //WITH_MLSDK
 }
@@ -611,8 +610,7 @@ void FMagicLeapCustomPresentVulkan::FinishRendering()
 #endif
 		}
 	}
-
-	Plugin->GetCurrentFrame().bBeginFrameSucceeded = false;
+  
 	Plugin->InitializeOldFrameFromRenderFrame();
 #endif // WITH_MLSDK
 }
