Copying //UE4/Partner-MagicLeap-Staging to Dev-VR (//UE4/Dev-VR)

diff --git a/Engine/Plugins/Lumin/MagicLeap/Source/MagicLeap/Private/InAppPurchaseComponent.cpp b/Engine/Plugins/Lumin/MagicLeap/Source/MagicLeap/Private/InAppPurchaseComponent.cpp
--- a/Engine/Plugins/Lumin/MagicLeap/Source/MagicLeap/Private/InAppPurchaseComponent.cpp
+++ b/Engine/Plugins/Lumin/MagicLeap/Source/MagicLeap/Private/InAppPurchaseComponent.cpp
@@ -9,9 +9,7 @@
 #include "Lumin/LuminPlatformAffinity.h"
 #endif // PLATFORM_LUMIN
 
-#if WITH_MLSDK
-#include <ml_purchase.h>
-#endif //WITH_MLSDK
+#include "Lumin/CAPIShims/LuminAPIPurchase.h"
 
 DEFINE_LOG_CATEGORY(LogInAppPurchase);
 
diff --git a/Engine/Plugins/Lumin/MagicLeap/Source/MagicLeap/Private/LuminApplicationLifecycleComponent.cpp b/Engine/Plugins/Lumin/MagicLeap/Source/MagicLeap/Private/LuminApplicationLifecycleComponent.cpp
new file mode 100644
--- /dev/null
+++ b/Engine/Plugins/Lumin/MagicLeap/Source/MagicLeap/Private/LuminApplicationLifecycleComponent.cpp
@@ -0,0 +1,28 @@
+// Copyright 1998-2019 Epic Games, Inc. All Rights Reserved.
+
+#include "LuminApplicationLifecycleComponent.h"
+#if PLATFORM_LUMIN
+#include "Lumin/LuminPlatformDelegates.h"
+#endif // PLATFORM_LUMIN
+
+void ULuminApplicationLifecycleComponent::OnRegister()
+{
+	Super::OnRegister();
+
+#if PLATFORM_LUMIN
+	FLuminDelegates::DeviceHasReactivatedDelegate.AddUObject(this, &ULuminApplicationLifecycleComponent::DeviceHasReactivatedDelegate_Handler);
+	FLuminDelegates::DeviceWillEnterRealityModeDelegate.AddUObject(this, &ULuminApplicationLifecycleComponent::DeviceWillEnterRealityModeDelegate_Handler);
+	FLuminDelegates::DeviceWillGoInStandbyDelegate.AddUObject(this, &ULuminApplicationLifecycleComponent::DeviceWillGoInStandbyDelegate_Handler);
+#endif // PLATFORM_LUMIN
+}
+
+void ULuminApplicationLifecycleComponent::OnUnregister()
+{
+	Super::OnUnregister();
+
+#if PLATFORM_LUMIN
+	FLuminDelegates::DeviceHasReactivatedDelegate.RemoveAll(this);
+	FLuminDelegates::DeviceWillEnterRealityModeDelegate.RemoveAll(this);
+	FLuminDelegates::DeviceWillGoInStandbyDelegate.RemoveAll(this);
+#endif // PLATFORM_LUMIN
+}
diff --git a/Engine/Plugins/Lumin/MagicLeap/Source/MagicLeap/Private/MagicLeapCustomPresent.cpp b/Engine/Plugins/Lumin/MagicLeap/Source/MagicLeap/Private/MagicLeapCustomPresent.cpp
--- a/Engine/Plugins/Lumin/MagicLeap/Source/MagicLeap/Private/MagicLeapCustomPresent.cpp
+++ b/Engine/Plugins/Lumin/MagicLeap/Source/MagicLeap/Private/MagicLeapCustomPresent.cpp
@@ -3,13 +3,14 @@
 #include "MagicLeapCustomPresent.h"
 #include "MagicLeapHMD.h"
 #include "RenderingThread.h"
-#include "MagicLeapPluginUtil.h" // for ML_INCLUDES_START/END
+#include "Lumin/CAPIShims/LuminAPILifecycle.h"
 
-#if WITH_MLSDK
-ML_INCLUDES_START
-#include <ml_lifecycle.h>
-ML_INCLUDES_END
-#endif //WITH_MLSDK
+#include "RHICommandList.h"
+#include "GlobalShader.h"
+#include "ScreenRendering.h"
+#include "CommonRenderResources.h"
+#include "PipelineStateCache.h"
+#include "ClearQuad.h"
 
 #if PLATFORM_WINDOWS || PLATFORM_LINUX || PLATFORM_LUMIN
 #include "OpenGLDrvPrivate.h"
@@ -19,101 +20,67 @@ ML_INCLUDES_END
 #include "MagicLeapGraphics.h"
 
 #if PLATFORM_WINDOWS || PLATFORM_LUMIN
-#include "VulkanRHIPrivate.h"
 #include "XRThreadUtils.h"
-#include "MagicLeapHelperVulkan.h"
 #endif // PLATFORM_WINDOWS || PLATFORM_LUMIN
 
-bool FMagicLeapCustomPresent::NeedsNativePresent()
-{
-	return (Plugin->GetWindowMirrorMode() > 0);
-}
-
-#if PLATFORM_WINDOWS
+#include "Containers/Union.h"
 
-FMagicLeapCustomPresentD3D11::FMagicLeapCustomPresentD3D11(FMagicLeapHMD* plugin)
-	: FMagicLeapCustomPresent(plugin)
-	, RenderTargetTexture(0)
-{
-}
+DECLARE_FLOAT_COUNTER_STAT(TEXT("Far Clipping Plane"), STAT_FarClip, STATGROUP_MagicLeap);
+DECLARE_FLOAT_COUNTER_STAT(TEXT("Timewarp Stabilization Depth"), STAT_StabilizationDepth, STATGROUP_MagicLeap);
+DECLARE_FLOAT_COUNTER_STAT(TEXT("Focus Distance"), STAT_FocusDistance, STATGROUP_MagicLeap);
+DECLARE_FLOAT_COUNTER_STAT(TEXT("Near Clipping Plane"), STAT_NearClip, STATGROUP_MagicLeap);
 
-void FMagicLeapCustomPresentD3D11::BeginRendering()
+FMagicLeapCustomPresent::FMagicLeapCustomPresent(FMagicLeapHMD* plugin)
+: FRHICustomPresent()
+, Plugin(plugin)
+, bNotifyLifecycleOfFirstPresent(true)
+, bCustomPresentIsSet(false)
+, PlatformAPILevel(FMagicLeapAPISetup::GetPlatformLevel())
+, HFOV(80)
+, VFOV(60)
 {
-	check(IsInRenderingThread());
-}
-
-void FMagicLeapCustomPresentD3D11::FinishRendering()
-{
-	check(IsInRenderingThread());
-}
-
-void FMagicLeapCustomPresentD3D11::Reset()
-{
-	if (IsInGameThread())
+	bool bManualCallToAppReady = false;
+	if (GConfig->GetBool(TEXT("/Script/LuminRuntimeSettings.LuminRuntimeSettings"), TEXT("bManualCallToAppReady"), bManualCallToAppReady, GEngineIni))
 	{
-		// Wait for all resources to be released
-		FlushRenderingCommands();
+		bNotifyLifecycleOfFirstPresent = !bManualCallToAppReady;
 	}
 }
 
-void FMagicLeapCustomPresentD3D11::Shutdown()
-{
-	Reset();
-}
-
-void FMagicLeapCustomPresentD3D11::UpdateViewport(const FViewport& Viewport, FRHIViewport* InViewportRHI)
+bool FMagicLeapCustomPresent::NeedsNativePresent()
 {
-	check(IsInGameThread());
-	check(InViewportRHI);
-
-	const FTexture2DRHIRef& RT = Viewport.GetRenderTargetTexture();
-	check(IsValidRef(RT));
-
-	// TODO Assign RenderTargetTexture here.
-	InViewportRHI->SetCustomPresent(this);
+	return (Plugin->GetWindowMirrorMode() > 0);
 }
 
-void FMagicLeapCustomPresentD3D11::UpdateViewport_RenderThread()
-{
-}
+void FMagicLeapCustomPresent::OnBackBufferResize()
+{}
 
-void FMagicLeapCustomPresentD3D11::OnBackBufferResize()
+bool FMagicLeapCustomPresent::Present(int32& SyncInterval)
 {
-}
-
-bool FMagicLeapCustomPresentD3D11::Present(int32& SyncInterval)
-{
-	check(IsInRenderingThread());
+	check(IsInRenderingThread() || IsInRHIThread());
 
 	// turn off VSync for the 'normal Present'.
 	SyncInterval = 0;
+	// We don't do any mirroring on Lumin as we render direct to the device only.
+#if PLATFORM_LUMIN
+	bool bHostPresent = false;
+#else
 	bool bHostPresent = Plugin->GetWindowMirrorMode() > 0;
+#endif
 
 	FinishRendering();
-	return bHostPresent;
-}
 
-#endif // PLATFORM_WINDOWS
-
-#if PLATFORM_MAC
-
-FMagicLeapCustomPresentMetal::FMagicLeapCustomPresentMetal(FMagicLeapHMD* plugin)
-	: FMagicLeapCustomPresent(plugin)
-	, RenderTargetTexture(0)
-{
-}
+	bCustomPresentIsSet = false;
 
-void FMagicLeapCustomPresentMetal::BeginRendering()
-{
-	check(IsInRenderingThread() || IsInRHIThread());
+	return bHostPresent;
 }
 
-void FMagicLeapCustomPresentMetal::FinishRendering()
+void FMagicLeapCustomPresent::GetFieldOfView(float& OutHFOVInDegrees, float& OutVFOVInDegrees)
 {
-	check(IsInRenderingThread() || IsInRHIThread());
+	OutHFOVInDegrees = static_cast<float>(FPlatformAtomics::AtomicRead(&HFOV));
+	OutVFOVInDegrees = static_cast<float>(FPlatformAtomics::AtomicRead(&VFOV));
 }
 
-void FMagicLeapCustomPresentMetal::Reset()
+void FMagicLeapCustomPresent::Reset()
 {
 	if (IsInGameThread())
 	{
@@ -122,555 +89,110 @@ void FMagicLeapCustomPresentMetal::Reset()
 	}
 }
 
-void FMagicLeapCustomPresentMetal::Shutdown()
+void FMagicLeapCustomPresent::Shutdown()
 {
 	Reset();
 }
 
-void FMagicLeapCustomPresentMetal::UpdateViewport(const FViewport& Viewport, FRHIViewport* InViewportRHI)
+template <typename CameraParamsType>
+void FMagicLeapCustomPresent::InitCameraParams(CameraParamsType & CameraParams, FTrackingFrame & Frame)
 {
-	check(IsInGameThread());
-	check(InViewportRHI);
-
-	const FTexture2DRHIRef& RT = Viewport.GetRenderTargetTexture();
-	check(IsValidRef(RT));
+	// TODO [Blake] : Need to see if we can use this newer matrix and override the view
+	// projection matrix (since they query GetStereoProjectionMatrix on the main thread)
+#if WITH_MLSDK
+	CameraParams.projection_type = Frame.ProjectionType;
+#endif // WITH_MLSDK
+	CameraParams.surface_scale = Frame.PixelDensity;
+	CameraParams.protected_surface = false;
+	GConfig->GetBool(TEXT("/Script/LuminRuntimeSettings.LuminRuntimeSettings"), TEXT("bProtectedContent"), CameraParams.protected_surface, GEngineIni);
 
-	// TODO Assign RenderTargetTexture here.
-	InViewportRHI->SetCustomPresent(this);
-}
+	// The near clipping plane is expected in meters despite what is documented in the header.
+	CameraParams.near_clip = GNearClippingPlane / Frame.WorldToMetersScale;
+	CameraParams.far_clip = Frame.FarClippingPlane / Frame.WorldToMetersScale;
+	CameraParams.focus_distance = Frame.FocusDistance / Frame.WorldToMetersScale;
 
-void FMagicLeapCustomPresentMetal::UpdateViewport_RenderThread()
-{
+	InitExtraCameraParams(CameraParams, Frame);
 }
 
-void FMagicLeapCustomPresentMetal::OnBackBufferResize()
-{
-}
+#if WITH_MLSDK
+template <>
+void FMagicLeapCustomPresent::InitExtraCameraParams(MLGraphicsFrameParams & CameraParams, FTrackingFrame & Frame)
+{}
 
-bool FMagicLeapCustomPresentMetal::Present(int32& SyncInterval)
+template <>
+void FMagicLeapCustomPresent::InitExtraCameraParams(MLGraphicsFrameParamsEx & CameraParams, FTrackingFrame & Frame)
 {
-	check(IsInRenderingThread() || IsInRHIThread());
+	GConfig->GetBool(TEXT("/Script/LuminRuntimeSettings.LuminRuntimeSettings"), TEXT("bFrameVignette"), CameraParams.vignette, GEngineIni);
 
-	// turn off VSync for the 'normal Present'.
-	SyncInterval = 0;
-	bool bHostPresent = Plugin->GetWindowMirrorMode() > 0;
-
-	FinishRendering();
-	return bHostPresent;
-}
-
-#endif // PLATFORM_MAC
-
-#if PLATFORM_WINDOWS || PLATFORM_LINUX || PLATFORM_LUMIN
-
-#define BEGIN_END_FRAME_BALANCE_HACK 0
-
-#if BEGIN_END_FRAME_BALANCE_HACK
-static int BalanceCounter = 0;
-static MLHandle BalanceHandles[4] = { 0,0,0,0 };
-static MLHandle BalancePrevFrameHandle = 0;
-#endif //BEGIN_END_FRAME_BALANCE_HACK
-
-FMagicLeapCustomPresentOpenGL::FMagicLeapCustomPresentOpenGL(FMagicLeapHMD* plugin)
-	: FMagicLeapCustomPresent(plugin)
-	, RenderTargetTexture(0)
-	, bFramebuffersValid(false)
-{
+	CameraParams.stabilization_depth = Frame.StabilizationDepth / Frame.WorldToMetersScale;
 }
+#endif // WITH_MLSDK
 
-void FMagicLeapCustomPresentOpenGL::BeginRendering()
+void FMagicLeapCustomPresent::BeginFrame(FTrackingFrame& Frame)
 {
 #if WITH_MLSDK
-	check(IsInRenderingThread());
-
-	FTrackingFrame& frame = Plugin->GetCurrentFrameMutable();
-	if (bCustomPresentIsSet)
+	if (bCustomPresentIsSet && !Plugin->IsRenderingPaused())
 	{
-		// TODO [Blake] : Need to see if we can use this newer matrix and override the view
-		// projection matrix (since they query GetStereoProjectionMatrix on the main thread)
-		MLGraphicsFrameParams camera_params;
-		MLResult Result = MLGraphicsInitFrameParams(&camera_params);
-		if (Result != MLResult_Ok)
-		{
-			UE_LOG(LogMagicLeap, Error, TEXT("MLGraphicsInitFrameParams failed with status %d"), Result);
-		}
-		camera_params.projection_type = MLGraphicsProjectionType_ReversedInfiniteZ;
-		camera_params.surface_scale = frame.PixelDensity;
-		camera_params.protected_surface = false;
-		GConfig->GetBool(TEXT("/Script/LuminRuntimeSettings.LuminRuntimeSettings"), TEXT("bProtectedContent"), camera_params.protected_surface, GEngineIni);
-
-		// The near clipping plane is expected in meters despite what is documented in the header.
-		camera_params.near_clip = GNearClippingPlane / frame.WorldToMetersScale;
-		camera_params.far_clip = frame.FarClippingPlane / frame.WorldToMetersScale;
-		// Only focus distance equaling 1 engine unit seems to work on board without wearable and on desktop.
-#if PLATFORM_LUMIN
-		camera_params.focus_distance = frame.FocusDistance / frame.WorldToMetersScale;
-#else
-		camera_params.focus_distance = 1.0f;
-#endif
-
-#if BEGIN_END_FRAME_BALANCE_HACK
-		if (BalanceCounter != 0)
-		{
-			UE_LOG(LogMagicLeap, Error, TEXT("Begin / End frame calls out of balance!"));
-			MLResult Result = MLGraphicsSignalSyncObjectGL(Plugin->GraphicsClient, BalanceHandles[0]);
-			if (Result != MLResult_Ok)
-			{
-				UE_LOG(LogMagicLeap, Error, TEXT("MLGraphicsSignalSyncObjectGL for BalanceHandles[0] failed with status %d"), Result);
-			}
-
-			Result = MLGraphicsSignalSyncObjectGL(Plugin->GraphicsClient, BalanceHandles[1]);
-			if (Result != MLResult_Ok)
-			{
-				UE_LOG(LogMagicLeap, Error, TEXT("MLGraphicsSignalSyncObjectGL for BalanceHandles[1] failed with status %d"), Result);
-			}
-
-			Result = MLGraphicsEndFrame(Plugin->GraphicsClient, BalancePrevFrameHandle);
-			if (Result != MLResult_Ok)
-			{
-				UE_LOG(LogMagicLeap, Error, TEXT("MLGraphicsEndFrame failed with status %d"), Result);
-			}
-			--BalanceCounter;
-		}
-#endif
-
-		Result = MLGraphicsBeginFrame(Plugin->GraphicsClient, &camera_params, &frame.Handle, &frame.RenderInfoArray);
-		frame.bBeginFrameSucceeded = (Result == MLResult_Ok);
-		if (frame.bBeginFrameSucceeded)
+		TUnion< MLGraphicsFrameParams, MLGraphicsFrameParamsEx > camera_params;
+		if (PlatformAPILevel >= 2)
 		{
-#if BEGIN_END_FRAME_BALANCE_HACK
-			++BalanceCounter;
-			BalancePrevFrameHandle = frame.Handle;
-			BalanceHandles[0] = frame.RenderInfoArray.virtual_cameras[0].sync_object;
-			BalanceHandles[1] = frame.RenderInfoArray.virtual_cameras[1].sync_object;
-#endif //BEGIN_END_FRAME_BALANCE_HACK
-
-			/* Convert eye extents from Graphics Projection Model to Unreal Projection Model */
-			// Unreal expects the projection matrix to be in centimeters and uses it for various purposes
-			// such as bounding volume calculations for lights in the shadow algorithm.
-			// We're overwriting the near value to match the units of unreal here instead of using the units of the SDK
-			for (uint32_t eye = 0; eye < frame.RenderInfoArray.num_virtual_cameras; ++eye)
-			{
-				frame.RenderInfoArray.virtual_cameras[eye].projection.matrix_colmajor[10] = 0.0f; // Model change hack
-				frame.RenderInfoArray.virtual_cameras[eye].projection.matrix_colmajor[11] = -1.0f; // Model change hack
-				frame.RenderInfoArray.virtual_cameras[eye].projection.matrix_colmajor[14] = GNearClippingPlane; // Model change hack
-			}
+			camera_params.SetSubtype<MLGraphicsFrameParamsEx>(MLGraphicsFrameParamsEx());
+			MLGraphicsFrameParamsExInit(&camera_params.GetSubtype<MLGraphicsFrameParamsEx>());
+			InitCameraParams(camera_params.GetSubtype<MLGraphicsFrameParamsEx>(), Frame);
+			SET_FLOAT_STAT(STAT_FarClip, camera_params.GetSubtype<MLGraphicsFrameParamsEx>().far_clip);
+			SET_FLOAT_STAT(STAT_StabilizationDepth, camera_params.GetSubtype<MLGraphicsFrameParamsEx>().stabilization_depth);
+			SET_FLOAT_STAT(STAT_FocusDistance, camera_params.GetSubtype<MLGraphicsFrameParamsEx>().focus_distance);
+			SET_FLOAT_STAT(STAT_NearClip, camera_params.GetSubtype<MLGraphicsFrameParamsEx>().near_clip);
 		}
 		else
 		{
-			if (Result != MLResult_Timeout)
-			{
-				UE_LOG(LogMagicLeap, Error, TEXT("MLGraphicsBeginFrame failed with status %d"), Result);
-			}
-			// TODO: See if this is only needed for ZI.
-			frame.Handle = ML_INVALID_HANDLE;
-			MagicLeap::ResetVirtualCameraInfoArray(frame.RenderInfoArray);
-		}
-	}
-#endif //WITH_MLSDK
-}
-
-void FMagicLeapCustomPresentOpenGL::FinishRendering()
-{
-#if WITH_MLSDK
-	check(IsInRenderingThread());
-
-	if (Plugin->IsDeviceInitialized() && Plugin->GetCurrentFrame().bBeginFrameSucceeded)
-	{
-		if (bNotifyLifecycleOfFirstPresent)
-		{
-			// Lifecycle tells the system's loading indicator to stop.
-			// App's rendering takes over.
-			MLResult Result = MLLifecycleSetReadyIndication();
-			bNotifyLifecycleOfFirstPresent = (Result != MLResult_Ok);
+			camera_params.SetSubtype<MLGraphicsFrameParams>(MLGraphicsFrameParams());
+			MLResult Result = MLGraphicsInitFrameParams(&camera_params.GetSubtype<MLGraphicsFrameParams>());
 			if (Result != MLResult_Ok)
 			{
-				UE_LOG(LogMagicLeap, Error, TEXT("Error (%d) sending app ready indication to lifecycle."), Result);
+				UE_LOG(LogMagicLeap, Error, TEXT("MLGraphicsInitFrameParams failed with status %d"), Result);
 			}
-			else
-			{
-				// [temporary] used for KPI tracking.
-				UE_LOG(LogMagicLeap, Display, TEXT("Presenting first render from app."));
-			}
-		}
-
-		// TODO [Blake] : Hack since we cannot yet specify a handle per view in the view family
-		const MLGraphicsVirtualCameraInfoArray& vp_array = Plugin->GetCurrentFrame().RenderInfoArray;
-		const uint32 vp_width = static_cast<uint32>(vp_array.viewport.w);
-		const uint32 vp_height = static_cast<uint32>(vp_array.viewport.h);
-
-		if (!bFramebuffersValid)
-		{
-			glGenFramebuffers(2, Framebuffers);
-			bFramebuffersValid = true;
-		}
-
-		GLint CurrentFB = 0;
-		glGetIntegerv(GL_DRAW_FRAMEBUFFER_BINDING, &CurrentFB);
-
-		GLint FramebufferSRGB = 0;
-		glGetIntegerv(GL_FRAMEBUFFER_SRGB, &FramebufferSRGB);
-		if (FramebufferSRGB)
+			InitCameraParams(camera_params.GetSubtype<MLGraphicsFrameParams>(), Frame);
+			SET_FLOAT_STAT(STAT_FarClip, camera_params.GetSubtype<MLGraphicsFrameParams>().far_clip);
+			SET_FLOAT_STAT(STAT_FocusDistance, camera_params.GetSubtype<MLGraphicsFrameParams>().focus_distance);
+			SET_FLOAT_STAT(STAT_NearClip, camera_params.GetSubtype<MLGraphicsFrameParams>().near_clip);
+		}		
+
+		MLResult Result = MLResult_Ok;
+		if (PlatformAPILevel >= 2)
 		{
-			glDisable(GL_FRAMEBUFFER_SRGB);
+			Result = MLGraphicsBeginFrameEx(Plugin->GraphicsClient, &camera_params.GetSubtype<MLGraphicsFrameParamsEx>(), &Frame.FrameInfo);
 		}
-
-		//check(vp_array.num_virtual_cameras >= 2); // We assume at least one virtual camera per eye
-
-		const FIntPoint& IdealRenderTargetSize = Plugin->GetHMDDevice()->GetIdealRenderTargetSize();
-		const int32 SizeX = FMath::CeilToInt(IdealRenderTargetSize.X * Plugin->GetCurrentFrame().PixelDensity);
-		const int32 SizeY = FMath::CeilToInt(IdealRenderTargetSize.Y * Plugin->GetCurrentFrame().PixelDensity);
-
-		// this texture contains both eye renders
-		glBindFramebuffer(GL_FRAMEBUFFER, Framebuffers[0]);
-		FOpenGL::FramebufferTexture2D(GL_FRAMEBUFFER, GL_COLOR_ATTACHMENT0, GL_TEXTURE_2D, RenderTargetTexture, 0);
-
-		glBindFramebuffer(GL_FRAMEBUFFER, Framebuffers[1]);
-		FOpenGL::FramebufferTextureLayer(GL_FRAMEBUFFER, GL_COLOR_ATTACHMENT0, vp_array.color_id, 0, 0);
-
-		glBindFramebuffer(GL_READ_FRAMEBUFFER, Framebuffers[0]);
-		glBindFramebuffer(GL_DRAW_FRAMEBUFFER, Framebuffers[1]);
-
-		const bool bShouldFlipVertically = !IsES2Platform(GMaxRHIShaderPlatform);
-
-		bShouldFlipVertically ?
-			FOpenGL::BlitFramebuffer(0, 0, SizeX / 2, SizeY, 0, vp_height, vp_width, 0, GL_COLOR_BUFFER_BIT, GL_NEAREST) :
-			FOpenGL::BlitFramebuffer(0, 0, SizeX / 2, SizeY, 0, 0, vp_width, vp_height, GL_COLOR_BUFFER_BIT, GL_NEAREST);
-
-		MLResult Result = MLGraphicsSignalSyncObjectGL(Plugin->GraphicsClient, vp_array.virtual_cameras[0].sync_object);
-		if (Result != MLResult_Ok)
-		{
-			UE_LOG(LogMagicLeap, Error, TEXT("MLGraphicsSignalSyncObjectGL for eye 0 failed with status %d"), Result);
-		}
-
-		FOpenGL::FramebufferTextureLayer(GL_FRAMEBUFFER, GL_COLOR_ATTACHMENT0, vp_array.color_id, 0, 1);
-
-		bShouldFlipVertically ?
-			FOpenGL::BlitFramebuffer(SizeX / 2, 0, SizeX, SizeY, 0, vp_height, vp_width, 0, GL_COLOR_BUFFER_BIT, GL_NEAREST) :
-			FOpenGL::BlitFramebuffer(SizeX / 2, 0, SizeX, SizeY, 0, 0, vp_width, vp_height, GL_COLOR_BUFFER_BIT, GL_NEAREST);
-
-		Result = MLGraphicsSignalSyncObjectGL(Plugin->GraphicsClient, vp_array.virtual_cameras[1].sync_object);
-		if (Result != MLResult_Ok)
-		{
-			UE_LOG(LogMagicLeap, Error, TEXT("MLGraphicsSignalSyncObjectGL for eye 1 failed with status %d"), Result);
-		}
-
-		glBindFramebuffer(GL_DRAW_FRAMEBUFFER, CurrentFB);
-		if (FramebufferSRGB)
+		else
 		{
-			glEnable(GL_FRAMEBUFFER_SRGB);
+			Result = MLGraphicsBeginFrame(Plugin->GraphicsClient, &camera_params.GetSubtype<MLGraphicsFrameParams>(), &Frame.FrameInfo.handle, &Frame.FrameInfo.virtual_camera_info_array);
 		}
-
-		static_assert(ARRAY_COUNT(vp_array.virtual_cameras) == 2, "The MLSDK has updated the size of the virtual_cameras array.");
-#if 0 // Enable this in case the MLSDK increases the size of the virtual_cameras array past 2
-		for (uint32 i = 2; i < vp_array.num_virtual_cameras; ++i)
+		Frame.bBeginFrameSucceeded = (Result == MLResult_Ok);
+		if (!Frame.bBeginFrameSucceeded)
 		{
-			Result = MLGraphicsSignalSyncObjectGL(Plugin->GraphicsClient, vp_array.virtual_cameras[i].sync_object);
-			if (Result != MLResult_Ok)
+			if (Result != MLResult_Timeout)
 			{
-				UE_LOG(LogMagicLeap, Error, TEXT("MLGraphicsSignalSyncObjectGL for eye %d failed with status %d"), i, Result);
+				UE_LOG(LogMagicLeap, Error, TEXT("MLGraphicsBeginFrame failed with status %s"), UTF8_TO_TCHAR(MLGetResultString(Result)));
 			}
+			// TODO: See if this is only needed for ZI.
+			Frame.FrameInfo.handle = ML_INVALID_HANDLE;
+			MagicLeap::ResetVirtualCameraInfoArray(Frame.FrameInfo.virtual_camera_info_array);
 		}
-#endif
 
-#if BEGIN_END_FRAME_BALANCE_HACK
-		--BalanceCounter;
-#endif //BEGIN_END_FRAME_BALANCE_HACK
-
-		Result = MLGraphicsEndFrame(Plugin->GraphicsClient, Plugin->GetCurrentFrame().Handle);
-		if (Result != MLResult_Ok)
-		{
-			UE_LOG(LogMagicLeap, Error, TEXT("MLGraphicsEndFrame failed with status %d"), Result);
-		}
+		const MLGraphicsVirtualCameraInfo& VCamInfo = Frame.FrameInfo.virtual_camera_info_array.virtual_cameras[0];
+		FPlatformAtomics::AtomicStore(&HFOV, static_cast<int64>(FMath::RadiansToDegrees((VCamInfo.left_half_angle + VCamInfo.right_half_angle)*2.0f)));
+		FPlatformAtomics::AtomicStore(&VFOV, static_cast<int64>(FMath::RadiansToDegrees((VCamInfo.top_half_angle + VCamInfo.bottom_half_angle)*2.0f)));
 	}
-	Plugin->InitializeOldFrameFromRenderFrame();
-#endif //WITH_MLSDK
-}
-
-void FMagicLeapCustomPresentOpenGL::Reset()
-{
-	if (IsInGameThread())
-	{
-		// Wait for all resources to be released
-		FlushRenderingCommands();
-	}
-	else if (IsInRenderingThread() && bFramebuffersValid)
+	else
 	{
-		glDeleteFramebuffers(2, Framebuffers);
-		bFramebuffersValid = false;
+		Frame.bBeginFrameSucceeded = false;
 	}
-}
-
-void FMagicLeapCustomPresentOpenGL::Shutdown()
-{
-	Reset();
-}
-
-void FMagicLeapCustomPresentOpenGL::UpdateViewport(const FViewport& Viewport, FRHIViewport* InViewportRHI)
-{
-	check(IsInGameThread());
-	check(InViewportRHI);
-
-	const FTexture2DRHIRef& RT = Viewport.GetRenderTargetTexture();
-	check(IsValidRef(RT));
-
-	RenderTargetTexture = *(reinterpret_cast<uint32_t*>(RT->GetNativeResource()));
-	InViewportRHI->SetCustomPresent(this);
-	
-	FMagicLeapCustomPresentOpenGL* CustomPresent = this;
-	ENQUEUE_RENDER_COMMAND(UpdateViewport_RT)(
-		[CustomPresent](FRHICommandListImmediate& RHICmdList)
-		{
-			CustomPresent->UpdateViewport_RenderThread();
-		}
-	);
-}
-
-void FMagicLeapCustomPresentOpenGL::UpdateViewport_RenderThread()
-{
-	check(IsInRenderingThread());
-	bCustomPresentIsSet = true;
-}
-
-void FMagicLeapCustomPresentOpenGL::OnBackBufferResize()
-{
-}
-
-bool FMagicLeapCustomPresentOpenGL::Present(int& SyncInterval)
-{
-	check(IsInRenderingThread());
-
-	// turn off VSync for the 'normal Present'.
-	SyncInterval = 0;
-	// We don't do any mirroring on Lumin as we render direct to the device only.
-#if PLATFORM_LUMIN
-	bool bHostPresent = false;
-#else
-	bool bHostPresent = Plugin->GetWindowMirrorMode() > 0;
-#endif
-
-	FinishRendering();
-
-	bCustomPresentIsSet = false;
-
-	return bHostPresent;
-}
-
-#endif  // PLATFORM_WINDOWS || PLATFORM_LINUX || PLATFORM_LUMIN
-
-#if PLATFORM_WINDOWS || PLATFORM_LUMIN
-
-FMagicLeapCustomPresentVulkan::FMagicLeapCustomPresentVulkan(FMagicLeapHMD* plugin)
-	: FMagicLeapCustomPresent(plugin)
-	, RenderTargetTexture(VK_NULL_HANDLE)
-	, RenderTargetTextureAllocation(VK_NULL_HANDLE)
-	, RenderTargetTextureAllocationOffset(0)
-	, RenderTargetTextureSRGB(VK_NULL_HANDLE)
-	, LastAliasedRenderTarget(VK_NULL_HANDLE)
-{
-}
-
-void FMagicLeapCustomPresentVulkan::BeginRendering()
-{
-#if WITH_MLSDK
-	check(IsInRenderingThread());
-
-	ExecuteOnRHIThread([this]()
-	{
-		// Always use RHITrackingFrame here, which is then copied to the RenderTrackingFrame. 
-		FTrackingFrame& RHIframe = Plugin->RHITrackingFrame;
-		if (bCustomPresentIsSet)
-		{
-			// TODO [Blake] : Need to see if we can use this newer matrix and override the view
-			// projection matrix (since they query GetStereoProjectionMatrix on the main thread)
-			MLGraphicsFrameParams camera_params;
-			MLResult InitResult = MLGraphicsInitFrameParams(&camera_params);
-			if (InitResult != MLResult_Ok)
-			{
-				UE_LOG(LogMagicLeap, Error, TEXT("MLGraphicsInitFrameParams failed with status %d"), InitResult);
-			}
-			camera_params.projection_type = MLGraphicsProjectionType_UnsignedZ;
-			camera_params.surface_scale = RHIframe.PixelDensity;
-			camera_params.protected_surface = false;
-			GConfig->GetBool(TEXT("/Script/LuminRuntimeSettings.LuminRuntimeSettings"), TEXT("bProtectedContent"), camera_params.protected_surface, GEngineIni);
-
-			// The near clipping plane is expected in meters despite what is documented in the header.
-			camera_params.near_clip = GNearClippingPlane / RHIframe.WorldToMetersScale;
-			camera_params.far_clip = RHIframe.FarClippingPlane / RHIframe.WorldToMetersScale;
-
-			// The focus distance is expected in meters despite what is documented in the header.
-			// Only focus distance equaling 1 engine unit seems to work on board without wearable and on desktop.
-#if PLATFORM_LUMIN
-			camera_params.focus_distance = RHIframe.FocusDistance / RHIframe.WorldToMetersScale;
-#else
-			camera_params.focus_distance = 1.0f;
-#endif
-
-#if BEGIN_END_FRAME_BALANCE_HACK
-			if (BalanceCounter != 0)
-			{
-				UE_LOG(LogMagicLeap, Error, TEXT("Begin / End frame calls out of balance!"));
-				FMagicLeapHelperVulkan::SignalObjects((uint64)BalanceHandles[0], (uint64)BalanceHandles[1]);
-				MLResult Result = MLGraphicsEndFrame(Plugin->GraphicsClient, BalancePrevFrameHandle);
-				if (Result != MLResult_Ok)
-				{
-					UE_LOG(LogMagicLeap, Error, TEXT("MLGraphicsEndFrame failed with status %d"), Result);
-				}
-				--BalanceCounter;
-			}
-#endif
-
-			MLResult Result = MLGraphicsBeginFrame(Plugin->GraphicsClient, &camera_params, &RHIframe.Handle, &RHIframe.RenderInfoArray);
-			RHIframe.bBeginFrameSucceeded = (Result == MLResult_Ok);
-			if (RHIframe.bBeginFrameSucceeded)
-			{
-#if BEGIN_END_FRAME_BALANCE_HACK
-				++BalanceCounter;
-				BalancePrevFrameHandle = RHIframe.Handle;
-				BalanceHandles[0] = RHIframe.RenderInfoArray.virtual_cameras[0].sync_object;
-				BalanceHandles[1] = RHIframe.RenderInfoArray.virtual_cameras[1].sync_object;
-#endif //BEGIN_END_FRAME_BALANCE_HACK
-
-			  /* Convert eye extents from Graphics Projection Model to Unreal Projection Model */
-			  // Unreal expects the projection matrix to be in centimeters and uses it for various purposes
-			  // such as bounding volume calculations for lights in the shadow algorithm.
-			  // We're overwriting the near value to match the units of unreal here instead of using the units of the SDK
-				for (uint32_t eye = 0; eye < RHIframe.RenderInfoArray.num_virtual_cameras; ++eye)
-				{
-					RHIframe.RenderInfoArray.virtual_cameras[eye].projection.matrix_colmajor[10] = 0.0f; // Model change hack
-					RHIframe.RenderInfoArray.virtual_cameras[eye].projection.matrix_colmajor[11] = -1.0f; // Model change hack
-					RHIframe.RenderInfoArray.virtual_cameras[eye].projection.matrix_colmajor[14] = GNearClippingPlane; // Model change hack
-				}
-			}
-			else
-			{
-				if (Result != MLResult_Timeout)
-				{
-				UE_LOG(LogMagicLeap, Error, TEXT("MLGraphicsBeginFrame failed with status %d"), Result);
-				}
-				// TODO: See if this is only needed for ZI.
-				RHIframe.Handle = ML_INVALID_HANDLE;
-				MagicLeap::ResetVirtualCameraInfoArray(RHIframe.RenderInfoArray);
-			}
-			Plugin->InitializeRenderFrameFromRHIFrame();
-		}
-	});
 #endif // WITH_MLSDK
 }
 
-void FMagicLeapCustomPresentVulkan::FinishRendering()
+void FMagicLeapCustomPresent::NotifyFirstRender()
 {
 #if WITH_MLSDK
-	check(IsInRenderingThread() || IsInRHIThread());
-
-	if (Plugin->IsDeviceInitialized() && Plugin->GetCurrentFrame().bBeginFrameSucceeded)
-	{
-		// TODO [Blake] : Hack since we cannot yet specify a handle per view in the view family
-		const MLGraphicsVirtualCameraInfoArray& vp_array = Plugin->GetCurrentFrame().RenderInfoArray;
-		const uint32 vp_width = static_cast<uint32>(vp_array.viewport.w);
-		const uint32 vp_height = static_cast<uint32>(vp_array.viewport.h);
-
-		bool bTestClear = false;
-		if (bTestClear)
-		{
-			FMagicLeapHelperVulkan::TestClear((uint64)vp_array.color_id);
-		}
-		else
-		{
-			// Alias the render target with an srgb image description for proper color space output.
-			if (RenderTargetTextureAllocation != VK_NULL_HANDLE && LastAliasedRenderTarget != RenderTargetTexture)
-			{
-				// SDKUNREAL-1135: ML remote image is corrupted on AMD hardware
-				if (!IsRHIDeviceAMD())
-				{
-					// TODO: If RenderTargetTextureSRGB is non-null, we're leaking the previous handle here. Also leaking on shutdown.
-					RenderTargetTextureSRGB = reinterpret_cast<void *>(FMagicLeapHelperVulkan::AliasImageSRGB((uint64)RenderTargetTextureAllocation, RenderTargetTextureAllocationOffset, vp_width * 2, vp_height));
-					check(RenderTargetTextureSRGB != VK_NULL_HANDLE);
-				}
-				LastAliasedRenderTarget = RenderTargetTexture;
-				UE_LOG(LogMagicLeap, Log, TEXT("Aliased render target for correct sRGB ouput."));
-			}
-
-			const VkImage FinalTarget = (RenderTargetTextureSRGB != VK_NULL_HANDLE) ? reinterpret_cast<const VkImage>(RenderTargetTextureSRGB) : reinterpret_cast<const VkImage>(RenderTargetTexture);
-			FMagicLeapHelperVulkan::BlitImage((uint64)FinalTarget, 0, 0, 0, 0, vp_width, vp_height, 1, (uint64)vp_array.color_id, 0, 0, 0, 0, vp_width, vp_height, 1);
-			FMagicLeapHelperVulkan::BlitImage((uint64)FinalTarget, 0, vp_width, 0, 0, vp_width, vp_height, 1, (uint64)vp_array.color_id, 1, 0, 0, 0, vp_width, vp_height, 1);
-		}
-
-		FMagicLeapHelperVulkan::SignalObjects((uint64)vp_array.virtual_cameras[0].sync_object, (uint64)vp_array.virtual_cameras[1].sync_object);
-
-#if BEGIN_END_FRAME_BALANCE_HACK
-		--BalanceCounter;
-#endif //BEGIN_END_FRAME_BALANCE_HACK
-
-		MLResult Result = MLGraphicsEndFrame(Plugin->GraphicsClient, Plugin->GetCurrentFrame().Handle);
-		if (Result != MLResult_Ok)
-		{
-#if !WITH_EDITOR
-			UE_LOG(LogMagicLeap, Error, TEXT("MLGraphicsEndFrame failed with status %d"), Result);
-#endif
-		}
-	}
-  
-	Plugin->InitializeOldFrameFromRenderFrame();
-#endif // WITH_MLSDK
-}
-
-void FMagicLeapCustomPresentVulkan::Reset()
-{
-	if (IsInGameThread())
-	{
-		// Wait for all resources to be released
-		FlushRenderingCommands();
-	}
-}
-
-void FMagicLeapCustomPresentVulkan::Shutdown()
-{
-	Reset();
-}
-
-void FMagicLeapCustomPresentVulkan::UpdateViewport(const FViewport& Viewport, FRHIViewport* InViewportRHI)
-{
-	check(IsInGameThread());
-	check(InViewportRHI);
-
-	const FTexture2DRHIRef& RT = Viewport.GetRenderTargetTexture();
-	check(IsValidRef(RT));
-
-	RenderTargetTexture = RT->GetNativeResource();
-	RenderTargetTextureAllocation = reinterpret_cast<void*>(static_cast<FVulkanTexture2D*>(RT->GetTexture2D())->Surface.GetAllocationHandle());
-	RenderTargetTextureAllocationOffset = static_cast<FVulkanTexture2D*>(RT->GetTexture2D())->Surface.GetAllocationOffset();
-
-	InViewportRHI->SetCustomPresent(this);
-	
-	FMagicLeapCustomPresentVulkan* CustomPresent = this;
-	ENQUEUE_RENDER_COMMAND(UpdateViewport_RT)(
-		[CustomPresent](FRHICommandListImmediate& RHICmdList)
-		{
-			CustomPresent->UpdateViewport_RenderThread();
-		}
-	);
-}
-
-void FMagicLeapCustomPresentVulkan::UpdateViewport_RenderThread()
-{
-	check(IsInRenderingThread());
-
-	ExecuteOnRHIThread_DoNotWait([this]()
-	{
-		bCustomPresentIsSet = true;
-	});
-}
-
-void FMagicLeapCustomPresentVulkan::OnBackBufferResize()
-{
-}
-
-bool FMagicLeapCustomPresentVulkan::Present(int& SyncInterval)
-{
-#if WITH_MLSDK
-	check(IsInRenderingThread() || IsInRHIThread());
-
 	if (bNotifyLifecycleOfFirstPresent)
 	{
 		MLResult Result = MLLifecycleSetReadyIndication();
@@ -685,24 +207,52 @@ bool FMagicLeapCustomPresentVulkan::Present(int& SyncInterval)
 			UE_LOG(LogMagicLeap, Display, TEXT("Presenting first render from app."));
 		}
 	}
+#endif // WITH_MLSDK
+}
 
-	// turn off VSync for the 'normal Present'.
-	SyncInterval = 0;
-	// We don't do any mirroring on Lumin as we render direct to the device only.
-#if PLATFORM_LUMIN || PLATFORM_LUMINGL4
-	bool bHostPresent = false;
-#else
-	bool bHostPresent = Plugin->GetWindowMirrorMode() > 0;
-#endif
-
-	FinishRendering();
-
-	bCustomPresentIsSet = false;
+void FMagicLeapCustomPresent::RenderToTextureSlice_RenderThread(FRHICommandListImmediate& RHICmdList, FRHITexture2D* SrcTexture, FRHITexture* DstTexture, uint32 ArraySlice, const FVector4& UVandSize)
+{
+#if WITH_MLSDK
+	FRHIRenderPassInfo RPInfo(DstTexture, ERenderTargetActions::DontLoad_Store, nullptr, 0, ArraySlice);
+	RHICmdList.BeginRenderPass(RPInfo, TEXT("MagicLeap_RenderToMLSurface"));
+	{
+		const MLGraphicsVirtualCameraInfoArray& vp_array = Plugin->GetCurrentFrame().FrameInfo.virtual_camera_info_array;
+		const uint32 vp_width = static_cast<uint32>(vp_array.viewport.w);
+		const uint32 vp_height = static_cast<uint32>(vp_array.viewport.h);
 
-	return bHostPresent;
-#else
-	return false;
+		RHICmdList.SetViewport(0, 0, 0, vp_width, vp_height, 1.0f);
+
+		FGraphicsPipelineStateInitializer GraphicsPSOInit;
+		RHICmdList.ApplyCachedRenderTargets(GraphicsPSOInit);
+
+		const auto FeatureLevel = GMaxRHIFeatureLevel;
+		auto ShaderMap = GetGlobalShaderMap(FeatureLevel);
+		TShaderMapRef<FScreenVS> VertexShader(ShaderMap);
+		TShaderMapRef<FScreenPS> PixelShader(ShaderMap);
+
+		GraphicsPSOInit.BlendState = TStaticBlendState<>::GetRHI();
+		GraphicsPSOInit.RasterizerState = TStaticRasterizerState<>::GetRHI();
+		GraphicsPSOInit.DepthStencilState = TStaticDepthStencilState<false, CF_Always>::GetRHI();
+		GraphicsPSOInit.BoundShaderState.VertexDeclarationRHI = GFilterVertexDeclaration.VertexDeclarationRHI;
+		GraphicsPSOInit.BoundShaderState.VertexShaderRHI = GETSAFERHISHADER_VERTEX(*VertexShader);
+		GraphicsPSOInit.BoundShaderState.PixelShaderRHI = GETSAFERHISHADER_PIXEL(*PixelShader);
+		GraphicsPSOInit.PrimitiveType = PT_TriangleList;
+
+		SetGraphicsPipelineState(RHICmdList, GraphicsPSOInit);
+
+		PixelShader->SetParameters(RHICmdList, TStaticSamplerState<SF_Bilinear>::GetRHI(), SrcTexture);
+
+		Plugin->GetRendererModule()->DrawRectangle(
+			RHICmdList,
+			0, 0,
+			vp_width, vp_height,
+			UVandSize.X, UVandSize.Y, 	// U, V
+			UVandSize.Z, UVandSize.W,	// SizeU, SizeV
+			FIntPoint(vp_width, vp_height),
+			FIntPoint(1, 1),
+			*VertexShader,
+			EDRF_Default);
+	}
+	RHICmdList.EndRenderPass();
 #endif // WITH_MLSDK
 }
-
-#endif // PLATFORM_WINDOWS || PLATFORM_LUMIN
diff --git a/Engine/Plugins/Lumin/MagicLeap/Source/MagicLeap/Private/MagicLeapCustomPresent.h b/Engine/Plugins/Lumin/MagicLeap/Source/MagicLeap/Private/MagicLeapCustomPresent.h
--- a/Engine/Plugins/Lumin/MagicLeap/Source/MagicLeap/Private/MagicLeapCustomPresent.h
+++ b/Engine/Plugins/Lumin/MagicLeap/Source/MagicLeap/Private/MagicLeapCustomPresent.h
@@ -8,14 +8,9 @@
 #include "MagicLeapGraphics.h"
 #include "MagicLeapMath.h"
 #include "MagicLeapUtils.h"
-#include "MagicLeapPluginUtil.h" // for ML_INCLUDES_START/END
-
-#if WITH_MLSDK
-ML_INCLUDES_START
-#include <ml_api.h>
-#include <ml_snapshot.h>
-ML_INCLUDES_END
-#endif //WITH_MLSDK
+#include "Lumin/CAPIShims/LuminAPIPerception.h"
+#include "MagicLeapPluginUtil.h"
+#include "Math/Vector4.h"
 
 class FMagicLeapHMD;
 class FViewport;
@@ -36,16 +31,15 @@ public:
 	float NearClippingPlane;
 	float FarClippingPlane;
 	float RecommendedFarClippingPlane;
-#if WITH_MLSDK
-	MLSnapshot* Snapshot;
-#endif //WITH_MLSDK
+	float StabilizationDepth;
 	bool bBeginFrameSucceeded;
 
 #if WITH_MLSDK
-	MLHandle Handle;
+	MLSnapshot* Snapshot;
 	MLCoordinateFrameUID FrameId;
-	MLGraphicsClipExtentsInfoArray UpdateInfoArray; // update information for the frame
-	MLGraphicsVirtualCameraInfoArray RenderInfoArray; // render information for the frame
+	MLGraphicsFrameInfo FrameInfo; // render information for the frame
+	MLGraphicsClipExtentsInfoArrayEx UpdateInfoArray; // update information for the frame
+	MLGraphicsProjectionType ProjectionType;
 #endif //WITH_MLSDK
 
 	float PixelDensity;
@@ -62,12 +56,11 @@ public:
 		, NearClippingPlane(GNearClippingPlane)
 		, FarClippingPlane(1000.0f) // 10m
 		, RecommendedFarClippingPlane(FarClippingPlane)
-#if WITH_MLSDK
-		, Snapshot(nullptr)
-#endif //WITH_MLSDK
+		, StabilizationDepth(1000.0f) // 10m
 		, bBeginFrameSucceeded(false)
 #if WITH_MLSDK
-		, Handle(ML_INVALID_HANDLE)
+		, Snapshot(nullptr)
+		, ProjectionType(MLGraphicsProjectionType_ReversedInfiniteZ)
 #endif //WITH_MLSDK
 		, PixelDensity(1.0f)
 		, WorldContext(nullptr)
@@ -77,7 +70,8 @@ public:
 		FrameId.data[1] = 0;
 
 		MagicLeap::ResetClipExtentsInfoArray(UpdateInfoArray);
-		MagicLeap::ResetVirtualCameraInfoArray(RenderInfoArray);
+		MLGraphicsFrameInfoInit(&FrameInfo);
+		MagicLeap::ResetVirtualCameraInfoArray(FrameInfo.virtual_camera_info_array);
 #endif //WITH_MLSDK
 	}
 };
@@ -85,121 +79,37 @@ public:
 class FMagicLeapCustomPresent : public FRHICustomPresent
 {
 public:
-	FMagicLeapCustomPresent(FMagicLeapHMD* plugin) :
-		FRHICustomPresent(),
-		Plugin(plugin),
-		bNeedReinitRendererAPI(true),
-		bNotifyLifecycleOfFirstPresent(true),
-		bCustomPresentIsSet(false)
-	{}
+	FMagicLeapCustomPresent(FMagicLeapHMD* plugin);
 
 	virtual void BeginRendering() = 0;
 	virtual void FinishRendering() = 0;
 	virtual bool NeedsNativePresent() override;
+	virtual void OnBackBufferResize() override;
+	virtual bool Present(int& SyncInterval) override;
 
 	virtual void UpdateViewport(const FViewport& Viewport, FRHIViewport* InViewportRHI) = 0;
 	virtual void UpdateViewport_RenderThread() = 0;
-	virtual void Reset() = 0;
-	virtual void Shutdown() = 0;
+	virtual void GetFieldOfView(float& OutHFOVInDegrees, float& OutVFOVInDegrees);
+	virtual void Reset();
+	virtual void Shutdown();
 
-	virtual void SetNeedReinitRendererAPI() { bNeedReinitRendererAPI = true; }
+	// Override if custom present needs to render to ML surfaces using Unreal's render pipeline instead of direct blit / texture copy.
+	virtual void RenderToMLSurfaces_RenderThread(class FRHICommandListImmediate& RHICmdList, class FRHITexture2D* SrcTexture) {}
 
 protected:
+	void BeginFrame(FTrackingFrame& Frame);
+	void NotifyFirstRender();
+	void RenderToTextureSlice_RenderThread(class FRHICommandListImmediate& RHICmdList, class FRHITexture2D* SrcTexture, class FRHITexture* DstTexture, uint32 ArraySlice, const FVector4& UVandSize);
 	FMagicLeapHMD* Plugin;
-	bool bNeedReinitRendererAPI;
 	bool bNotifyLifecycleOfFirstPresent;
 	bool bCustomPresentIsSet;
-};
+	uint32 PlatformAPILevel;
+	volatile int64 HFOV;
+	volatile int64 VFOV;
 
-#if PLATFORM_WINDOWS
-class FMagicLeapCustomPresentD3D11 : public FMagicLeapCustomPresent
-{
-public:
-	FMagicLeapCustomPresentD3D11(FMagicLeapHMD* plugin);
+	template <typename CameraParamsType>
+	void InitCameraParams(CameraParamsType & CameraParams, FTrackingFrame & Frame);
 
-	virtual void OnBackBufferResize() override;
-	virtual bool Present(int& SyncInterval) override;
-
-	virtual void BeginRendering() override;
-	virtual void FinishRendering() override;
-	virtual void UpdateViewport(const FViewport& Viewport, FRHIViewport* InViewportRHI) override;
-	virtual void UpdateViewport_RenderThread() override;
-	virtual void Reset() override;
-	virtual void Shutdown() override;
-
-protected:
-	uint32_t RenderTargetTexture = 0;
+	template <typename CameraParamsType>
+	void InitExtraCameraParams(CameraParamsType & CameraParams, FTrackingFrame & Frame);
 };
-#endif // PLATFORM_WINDOWS
-
-#if PLATFORM_MAC
-class FMagicLeapCustomPresentMetal : public FMagicLeapCustomPresent
-{
-public:
-	FMagicLeapCustomPresentMetal(FMagicLeapHMD* plugin);
-
-	virtual void OnBackBufferResize() override;
-	virtual bool Present(int& SyncInterval) override;
-
-	virtual void BeginRendering() override;
-	virtual void FinishRendering() override;
-	virtual void UpdateViewport(const FViewport& Viewport, FRHIViewport* InViewportRHI) override;
-	virtual void UpdateViewport_RenderThread() override;
-	virtual void Reset() override;
-	virtual void Shutdown() override;
-
-protected:
-	uint32_t RenderTargetTexture = 0;
-};
-#endif // PLATFORM_MAC
-
-#if PLATFORM_WINDOWS || PLATFORM_LINUX || PLATFORM_LUMIN
-class FMagicLeapCustomPresentOpenGL : public FMagicLeapCustomPresent
-{
-public:
-	FMagicLeapCustomPresentOpenGL(FMagicLeapHMD* plugin);
-
-	virtual void OnBackBufferResize() override;
-	virtual bool Present(int& SyncInterval) override;
-
-	virtual void BeginRendering() override;
-	virtual void FinishRendering() override;
-	virtual void UpdateViewport(const FViewport& Viewport, FRHIViewport* InViewportRHI) override;
-	virtual void UpdateViewport_RenderThread() override;
-	virtual void Reset() override;
-	virtual void Shutdown() override;
-
-protected:
-	uint32_t RenderTargetTexture = 0;
-	uint32_t Framebuffers[2];
-	bool bFramebuffersValid;
-};
-#endif // PLATFORM_WINDOWS || PLATFORM_LINUX || PLATFORM_LUMIN
-
-#if PLATFORM_WINDOWS || PLATFORM_LUMIN
-class FMagicLeapCustomPresentVulkan : public FMagicLeapCustomPresent
-{
-public:
-	FMagicLeapCustomPresentVulkan(FMagicLeapHMD* plugin);
-
-	virtual void OnBackBufferResize() override;
-	virtual bool Present(int& SyncInterval) override;
-
-	virtual void BeginRendering() override;
-	virtual void FinishRendering() override;
-	virtual void UpdateViewport(const FViewport& Viewport, FRHIViewport* InViewportRHI) override;
-	virtual void UpdateViewport_RenderThread() override;
-	virtual void Reset() override;
-	virtual void Shutdown() override;
-
-protected:
-	void* RenderTargetTexture;
-	void* RenderTargetTextureAllocation;
-	uint64 RenderTargetTextureAllocationOffset = 0;
-	void* RenderTargetTextureSRGB;
-	void* LastAliasedRenderTarget;
-};
-#endif // PLATFORM_LUMIN
-
-
-void CaptureCallStack(FString& OutCallstack);
diff --git a/Engine/Plugins/Lumin/MagicLeap/Source/MagicLeap/Private/MagicLeapCustomPresentGL.cpp b/Engine/Plugins/Lumin/MagicLeap/Source/MagicLeap/Private/MagicLeapCustomPresentGL.cpp
new file mode 100644
--- /dev/null
+++ b/Engine/Plugins/Lumin/MagicLeap/Source/MagicLeap/Private/MagicLeapCustomPresentGL.cpp
@@ -0,0 +1,177 @@
+// Copyright 1998-2019 Epic Games, Inc. All Rights Reserved.
+
+#include "MagicLeapCustomPresentGL.h"
+#include "MagicLeapHMD.h"
+#include "RenderingThread.h"
+#include "Lumin/CAPIShims/LuminAPI.h"
+
+#if PLATFORM_WINDOWS || PLATFORM_LINUX || PLATFORM_LUMIN
+#include "OpenGLDrvPrivate.h"
+#include "MagicLeapHelperOpenGL.h"
+#endif // PLATFORM_WINDOWS || PLATFORM_LINUX || PLATFORM_LUMIN
+
+#include "MagicLeapGraphics.h"
+#include "Containers/Union.h"
+
+#if PLATFORM_WINDOWS || PLATFORM_LINUX || PLATFORM_LUMIN
+
+FMagicLeapCustomPresentOpenGL::FMagicLeapCustomPresentOpenGL(FMagicLeapHMD* plugin)
+: FMagicLeapCustomPresent(plugin)
+, RenderTargetTexture(0)
+, bFramebuffersValid(false)
+{}
+
+void FMagicLeapCustomPresentOpenGL::BeginRendering()
+{
+#if WITH_MLSDK
+	check(IsInRenderingThread());
+
+	FTrackingFrame& frame = Plugin->GetCurrentFrameMutable();
+	BeginFrame(frame);
+#endif //WITH_MLSDK
+}
+
+void FMagicLeapCustomPresentOpenGL::FinishRendering()
+{
+#if WITH_MLSDK
+	check(IsInRenderingThread());
+
+	if (Plugin->IsDeviceInitialized() && Plugin->GetCurrentFrame().bBeginFrameSucceeded)
+	{
+		NotifyFirstRender();
+
+		// TODO [Blake] : Hack since we cannot yet specify a handle per view in the view family
+		const MLGraphicsVirtualCameraInfoArray& vp_array = Plugin->GetCurrentFrame().FrameInfo.virtual_camera_info_array;
+		const uint32 vp_width = static_cast<uint32>(vp_array.viewport.w);
+		const uint32 vp_height = static_cast<uint32>(vp_array.viewport.h);
+
+		if (!bFramebuffersValid)
+		{
+			glGenFramebuffers(2, Framebuffers);
+			bFramebuffersValid = true;
+		}
+
+		GLint CurrentFB = 0;
+		glGetIntegerv(GL_DRAW_FRAMEBUFFER_BINDING, &CurrentFB);
+
+		GLint FramebufferSRGB = 0;
+		glGetIntegerv(GL_FRAMEBUFFER_SRGB, &FramebufferSRGB);
+		if (FramebufferSRGB)
+		{
+			// Unreal render target texture is marked as linear but it actually contains pixels in sRGB color space.
+			// Rendering from a linear color space texture to an sRGB texture, when GL_FRAMEBUFFER_SRGB is enabled, causes implicit color conversion by the graphics library.
+			// This causes a double gamma correction on the pixels and the final eye texture rendered on the device / simulator appears very bright and washed out.
+			// We cannot reslove this by requesting linear color space eye textures from ml_graphics as the implicit conversion happens further down the compositor pipeline.
+			// The issue is fixed by disabling GL_FRAMEBUFFER_SRGB for the texture blit operation.
+			glDisable(GL_FRAMEBUFFER_SRGB);
+		}
+
+		//check(vp_array.num_virtual_cameras >= 2); // We assume at least one virtual camera per eye
+
+		const FIntPoint& IdealRenderTargetSize = Plugin->GetHMDDevice()->GetIdealRenderTargetSize();
+		const int32 SizeX = FMath::CeilToInt(IdealRenderTargetSize.X * Plugin->GetCurrentFrame().PixelDensity);
+		const int32 SizeY = FMath::CeilToInt(IdealRenderTargetSize.Y * Plugin->GetCurrentFrame().PixelDensity);
+
+		// this texture contains both eye renders
+		glBindFramebuffer(GL_FRAMEBUFFER, Framebuffers[0]);
+		FOpenGL::FramebufferTexture2D(GL_FRAMEBUFFER, GL_COLOR_ATTACHMENT0, GL_TEXTURE_2D, RenderTargetTexture, 0);
+
+		glBindFramebuffer(GL_FRAMEBUFFER, Framebuffers[1]);
+		FOpenGL::FramebufferTextureLayer(GL_FRAMEBUFFER, GL_COLOR_ATTACHMENT0, vp_array.color_id, 0, 0);
+
+		glBindFramebuffer(GL_READ_FRAMEBUFFER, Framebuffers[0]);
+		glBindFramebuffer(GL_DRAW_FRAMEBUFFER, Framebuffers[1]);
+
+		const bool bShouldFlipVertically = !IsES2Platform(GMaxRHIShaderPlatform);
+
+		bShouldFlipVertically ?
+			FOpenGL::BlitFramebuffer(0, 0, SizeX / 2, SizeY, 0, vp_height, vp_width, 0, GL_COLOR_BUFFER_BIT, GL_NEAREST) :
+			FOpenGL::BlitFramebuffer(0, 0, SizeX / 2, SizeY, 0, 0, vp_width, vp_height, GL_COLOR_BUFFER_BIT, GL_NEAREST);
+
+		MLResult Result = MLGraphicsSignalSyncObjectGL(Plugin->GraphicsClient, vp_array.virtual_cameras[0].sync_object);
+		if (Result != MLResult_Ok)
+		{
+			UE_LOG(LogMagicLeap, Error, TEXT("MLGraphicsSignalSyncObjectGL for eye 0 failed with status %d"), Result);
+		}
+
+		FOpenGL::FramebufferTextureLayer(GL_FRAMEBUFFER, GL_COLOR_ATTACHMENT0, vp_array.color_id, 0, 1);
+
+		bShouldFlipVertically ?
+			FOpenGL::BlitFramebuffer(SizeX / 2, 0, SizeX, SizeY, 0, vp_height, vp_width, 0, GL_COLOR_BUFFER_BIT, GL_NEAREST) :
+			FOpenGL::BlitFramebuffer(SizeX / 2, 0, SizeX, SizeY, 0, 0, vp_width, vp_height, GL_COLOR_BUFFER_BIT, GL_NEAREST);
+
+		Result = MLGraphicsSignalSyncObjectGL(Plugin->GraphicsClient, vp_array.virtual_cameras[1].sync_object);
+		if (Result != MLResult_Ok)
+		{
+			UE_LOG(LogMagicLeap, Error, TEXT("MLGraphicsSignalSyncObjectGL for eye 1 failed with status %d"), Result);
+		}
+
+		glBindFramebuffer(GL_DRAW_FRAMEBUFFER, CurrentFB);
+		if (FramebufferSRGB)
+		{
+			glEnable(GL_FRAMEBUFFER_SRGB);
+		}
+
+		static_assert(ARRAY_COUNT(vp_array.virtual_cameras) == 2, "The MLSDK has updated the size of the virtual_cameras array.");
+#if 0 // Enable this in case the MLSDK increases the size of the virtual_cameras array past 2
+		for (uint32 i = 2; i < vp_array.num_virtual_cameras; ++i)
+		{
+			Result = MLGraphicsSignalSyncObjectGL(Plugin->GraphicsClient, vp_array.virtual_cameras[i].sync_object);
+			if (Result != MLResult_Ok)
+			{
+				UE_LOG(LogMagicLeap, Error, TEXT("MLGraphicsSignalSyncObjectGL for eye %d failed with status %d"), i, Result);
+			}
+		}
+#endif
+
+		Result = MLGraphicsEndFrame(Plugin->GraphicsClient, Plugin->GetCurrentFrame().FrameInfo.handle);
+		if (Result != MLResult_Ok)
+		{
+			UE_LOG(LogMagicLeap, Error, TEXT("MLGraphicsEndFrame failed with status %d"), Result);
+		}
+	}
+	Plugin->InitializeOldFrameFromRenderFrame();
+#endif //WITH_MLSDK
+}
+
+void FMagicLeapCustomPresentOpenGL::Reset()
+{
+	if (IsInGameThread())
+	{
+		// Wait for all resources to be released
+		FlushRenderingCommands();
+	}
+	else if (IsInRenderingThread() && bFramebuffersValid)
+	{
+		glDeleteFramebuffers(2, Framebuffers);
+		bFramebuffersValid = false;
+	}
+}
+
+void FMagicLeapCustomPresentOpenGL::UpdateViewport(const FViewport& Viewport, FRHIViewport* InViewportRHI)
+{
+	check(IsInGameThread());
+	check(InViewportRHI);
+
+	const FTexture2DRHIRef& RT = Viewport.GetRenderTargetTexture();
+	check(IsValidRef(RT));
+
+	RenderTargetTexture = *(reinterpret_cast<uint32_t*>(RT->GetNativeResource()));
+	InViewportRHI->SetCustomPresent(this);
+	
+	FMagicLeapCustomPresentOpenGL* CustomPresent = this;
+	ENQUEUE_RENDER_COMMAND(UpdateViewport_RT)(
+		[CustomPresent](FRHICommandList& RHICmdList)
+		{
+			CustomPresent->UpdateViewport_RenderThread();
+		}
+	);
+}
+
+void FMagicLeapCustomPresentOpenGL::UpdateViewport_RenderThread()
+{
+	check(IsInRenderingThread());
+	bCustomPresentIsSet = true;
+}
+
+#endif  // PLATFORM_WINDOWS || PLATFORM_LINUX || PLATFORM_LUMIN
diff --git a/Engine/Plugins/Lumin/MagicLeap/Source/MagicLeap/Private/MagicLeapCustomPresentGL.h b/Engine/Plugins/Lumin/MagicLeap/Source/MagicLeap/Private/MagicLeapCustomPresentGL.h
new file mode 100644
--- /dev/null
+++ b/Engine/Plugins/Lumin/MagicLeap/Source/MagicLeap/Private/MagicLeapCustomPresentGL.h
@@ -0,0 +1,24 @@
+// Copyright 1998-2019 Epic Games, Inc. All Rights Reserved.
+
+#pragma once
+
+#include "MagicLeapCustomPresent.h"
+
+#if PLATFORM_WINDOWS || PLATFORM_LINUX || PLATFORM_LUMIN
+class FMagicLeapCustomPresentOpenGL : public FMagicLeapCustomPresent
+{
+public:
+	FMagicLeapCustomPresentOpenGL(FMagicLeapHMD* plugin);
+
+	virtual void BeginRendering() override;
+	virtual void FinishRendering() override;
+	virtual void UpdateViewport(const FViewport& Viewport, FRHIViewport* InViewportRHI) override;
+	virtual void UpdateViewport_RenderThread() override;
+	virtual void Reset() override;
+
+protected:
+	uint32_t RenderTargetTexture = 0;
+	uint32_t Framebuffers[2];
+	bool bFramebuffersValid;
+};
+#endif // PLATFORM_WINDOWS || PLATFORM_LINUX || PLATFORM_LUMIN
diff --git a/Engine/Plugins/Lumin/MagicLeap/Source/MagicLeap/Private/MagicLeapCustomPresentMetal.cpp b/Engine/Plugins/Lumin/MagicLeap/Source/MagicLeap/Private/MagicLeapCustomPresentMetal.cpp
new file mode 100644
--- /dev/null
+++ b/Engine/Plugins/Lumin/MagicLeap/Source/MagicLeap/Private/MagicLeapCustomPresentMetal.cpp
@@ -0,0 +1,164 @@
+// Copyright 1998-2018 Epic Games, Inc. All Rights Reserved.
+
+#include "MagicLeapCustomPresentMetal.h"
+#include "MagicLeapHMD.h"
+#include "RenderingThread.h"
+#include "Lumin/CAPIShims/LuminAPIRemote.h"
+#include "XRThreadUtils.h"
+
+#include "Containers/Union.h"
+
+#if PLATFORM_MAC
+
+#include "MetalRHI.h"
+#include "MetalRHIPrivate.h"
+#include "MetalContext.h"
+#include "MetalState.h"
+#include "MetalResources.h"
+
+FMagicLeapCustomPresentMetal::FMagicLeapCustomPresentMetal(FMagicLeapHMD* plugin)
+: FMagicLeapCustomPresent(plugin)
+, DestTextureRef(nullptr)
+, SrcTextureSRGBRef(nullptr)
+{}
+
+void FMagicLeapCustomPresentMetal::BeginRendering()
+{
+#if WITH_MLSDK
+	check(IsInRenderingThread());
+
+	ExecuteOnRHIThread([this]()
+	{
+		// Always use RHITrackingFrame here, which is then copied to the RenderTrackingFrame. 
+		FTrackingFrame& RHIframe = Plugin->RHITrackingFrame;
+		RHIframe.ProjectionType = MLGraphicsProjectionType_UnsignedZ;
+		BeginFrame(RHIframe);
+		if (bCustomPresentIsSet)
+		{
+			Plugin->InitializeRenderFrameFromRHIFrame();
+		}
+	});
+#endif // WITH_MLSDK
+}
+
+void FMagicLeapCustomPresentMetal::FinishRendering()
+{
+#if WITH_MLSDK
+	check(IsInRenderingThread() || IsInRHIThread());
+
+	if (Plugin->IsDeviceInitialized() && Plugin->GetCurrentFrame().bBeginFrameSucceeded)
+	{
+		// Actual rendering on the graphics textures happens in FMagicLeapCustomPresent::RenderToMLSurfaces_RenderThread() called via FMagicLeapHMD::RenderTexture_RenderThread()
+		// FMagicLeapCustomPresent::FinishRendering() can be called on the RHI thread as well and 
+		// we should setup the render pass on the render thread only.
+
+		// Notify first render here instead of in RenderToMLSurfaces_RenderThread() because the render to MLSurfaces is already finished by now.
+		NotifyFirstRender();
+
+		const MLGraphicsVirtualCameraInfoArray& vp_array = Plugin->GetCurrentFrame().FrameInfo.virtual_camera_info_array;
+
+		MLResult Result;
+		for (uint32 i = 0; i < vp_array.num_virtual_cameras; ++i)
+		{
+			Result = MLRemoteGraphicsSignalSyncObjectMTL(Plugin->GraphicsClient, vp_array.virtual_cameras[i].sync_object);
+			UE_CLOG(Result != MLResult_Ok, LogMagicLeap, Error, TEXT("MLRemoteGraphicsSignalSyncObjectMTL(%d) failed with status %d"), i, Result);
+		}
+
+		Result = MLGraphicsEndFrame(Plugin->GraphicsClient, Plugin->GetCurrentFrame().FrameInfo.handle);
+		if (Result != MLResult_Ok)
+		{
+#if !WITH_EDITOR
+			UE_LOG(LogMagicLeap, Error, TEXT("MLGraphicsEndFrame failed with status %d"), Result);
+#endif
+		}
+	}
+  
+	Plugin->InitializeOldFrameFromRenderFrame();
+#endif // WITH_MLSDK
+}
+
+void FMagicLeapCustomPresentMetal::UpdateViewport(const FViewport& Viewport, FRHIViewport* InViewportRHI)
+{
+	check(IsInGameThread());
+
+	InViewportRHI->SetCustomPresent(this);
+
+	FMagicLeapCustomPresentMetal* CustomPresent = this;
+	ENQUEUE_RENDER_COMMAND(UpdateViewport_RT)(
+		[CustomPresent](FRHICommandList& RHICmdList)
+		{
+			CustomPresent->UpdateViewport_RenderThread();
+		}
+	);
+}
+
+void FMagicLeapCustomPresentMetal::UpdateViewport_RenderThread()
+{
+	check(IsInRenderingThread());
+
+	ExecuteOnRHIThread_DoNotWait([this]()
+	{
+		bCustomPresentIsSet = true;
+	});
+}
+
+void FMagicLeapCustomPresentMetal::Reset()
+{
+	FMagicLeapCustomPresent::Reset();
+	if (IsInRenderingThread() && DestTextureRef.IsValid())
+	{
+		if (DestTextureRef.IsValid())
+		{
+			DestTextureRef.SafeRelease();
+		}
+		if (SrcTextureSRGBRef.IsValid())
+		{
+			SrcTextureSRGBRef.SafeRelease();
+		}
+	}
+}
+
+void FMagicLeapCustomPresentMetal::RenderToMLSurfaces_RenderThread(FRHICommandListImmediate& RHICmdList, FRHITexture2D* SrcTexture)
+{
+#if WITH_MLSDK
+	if (Plugin->IsDeviceInitialized() && Plugin->GetCurrentFrame().bBeginFrameSucceeded)
+	{
+		const MLGraphicsVirtualCameraInfoArray& vp_array = Plugin->GetCurrentFrame().FrameInfo.virtual_camera_info_array;
+
+		// Manage MLGraphics texture as an FMetalTexture to use it with Unreal's render pipeline.
+		// Use AutoRelease for ownership so that color_id is not owned by the FMetalTexture object, lifetime is externally managed.
+		FMetalTexture DestMTLTexture((id<MTLTexture>)vp_array.color_id, ns::Ownership::AutoRelease);
+		if (!DestTextureRef.IsValid())
+		{
+			const uint32 vp_width = static_cast<uint32>(vp_array.viewport.w);
+			const uint32 vp_height = static_cast<uint32>(vp_array.viewport.h);
+			// TODO: @njain add an overload to FMetalSurface to handle externally owned textures, similar to FVulkanSurface.
+			DestTextureRef = new FMetalTexture2DArray(PF_R8G8B8A8, vp_width, vp_height, vp_array.num_virtual_cameras, 1, TexCreate_RenderTargetable, nullptr, FClearValueBinding::Transparent);
+		}
+
+		// Set MLGraphics texture into Unreal's FMetalTexture2DArray.
+		static_cast<FMetalTexture2DArray*>(DestTextureRef.GetReference())->Surface.Texture = DestMTLTexture;
+
+		FMetalTexture2D* SrcMetalTexture = static_cast<FMetalTexture2D*>(SrcTexture);
+		if (!SrcTextureSRGBRef.IsValid())
+		{
+			SrcTextureSRGBRef = new FMetalTexture2D(SrcMetalTexture->GetFormat(), SrcMetalTexture->GetSizeX(), SrcMetalTexture->GetSizeY(), SrcMetalTexture->GetNumMips(), SrcMetalTexture->GetNumSamples(), SrcMetalTexture->GetFlags(), nullptr, FClearValueBinding::Transparent);
+		}
+
+		// Unreal render target texture is marked as linear but it actually contains pixels in sRGB color space.
+		// Rendering from a linear color space texture to an sRGB texture causes implicit color conversion by the graphics library.
+		// This causes a double gamma correction on the pixels and the final eye texture rendered on the device / simulator appears very bright and washed out.
+		// We cannot reslove this by requesting linear color space eye textures from ml_graphics as the implicit conversion happens further down the compositor pipeline.
+		// The issue is fixed by creating a new texture marked as sRGB but which uses the same memory as the original Unreal render target texture.
+		// We use this new sRGB texture to render to ml_graphics' eye textures.
+		static_cast<FMetalTexture2D*>(SrcTextureSRGBRef.GetReference())->Surface.Texture = FMetalTexture(SrcMetalTexture->Surface.Texture.NewTextureView(static_cast<mtlpp::PixelFormat>(MLRemoteGraphicsMTLFormatFromMLSurfaceFormat(MLSurfaceFormat_RGBA8UNormSRGB))));
+
+		// Left Eye (flip vertically) = U1 = 0, V1 = 1, U2 = (0 + 0.5), V2 = (1 - 1)
+		RenderToTextureSlice_RenderThread(RHICmdList, SrcTextureSRGBRef.GetReference(), DestTextureRef.GetReference(), 0, FVector4(0, 1.0f, 0.5f, -1.0f));
+		// Right Eye (flip vertically) = U1 = 0.5, V1 = 1, U2 = (0.5 + 0.5), V2 = (1 - 1)
+		RenderToTextureSlice_RenderThread(RHICmdList, SrcTextureSRGBRef.GetReference(), DestTextureRef.GetReference(), 1, FVector4(0.5f, 1.0f, 0.5f, -1.0f));
+	}
+#endif // WITH_MLSDK
+}
+
+#endif // PLATFORM_MAC
diff --git a/Engine/Plugins/Lumin/MagicLeap/Source/MagicLeap/Private/MagicLeapCustomPresentMetal.h b/Engine/Plugins/Lumin/MagicLeap/Source/MagicLeap/Private/MagicLeapCustomPresentMetal.h
new file mode 100644
--- /dev/null
+++ b/Engine/Plugins/Lumin/MagicLeap/Source/MagicLeap/Private/MagicLeapCustomPresentMetal.h
@@ -0,0 +1,27 @@
+// Copyright 1998-2019 Epic Games, Inc. All Rights Reserved.
+
+#pragma once
+
+#include "MagicLeapCustomPresent.h"
+
+#if PLATFORM_MAC
+
+class FMagicLeapCustomPresentMetal : public FMagicLeapCustomPresent
+{
+public:
+	FMagicLeapCustomPresentMetal(FMagicLeapHMD* plugin);
+
+	virtual void BeginRendering() override;
+	virtual void FinishRendering() override;
+	virtual void UpdateViewport(const FViewport& Viewport, FRHIViewport* InViewportRHI) override;
+	virtual void UpdateViewport_RenderThread() override;
+	virtual void Reset() override;
+
+	virtual void RenderToMLSurfaces_RenderThread(class FRHICommandListImmediate& RHICmdList, class FRHITexture2D* SrcTexture) override;
+
+protected:
+	// Used as a container for the MLGraphics texture array.
+	FTexture2DArrayRHIRef DestTextureRef;
+	FTexture2DRHIRef SrcTextureSRGBRef;
+};
+#endif // PLATFORM_MAC
diff --git a/Engine/Plugins/Lumin/MagicLeap/Source/MagicLeap/Private/MagicLeapCustomPresentVk.cpp b/Engine/Plugins/Lumin/MagicLeap/Source/MagicLeap/Private/MagicLeapCustomPresentVk.cpp
new file mode 100644
--- /dev/null
+++ b/Engine/Plugins/Lumin/MagicLeap/Source/MagicLeap/Private/MagicLeapCustomPresentVk.cpp
@@ -0,0 +1,152 @@
+// Copyright 1998-2019 Epic Games, Inc. All Rights Reserved.
+
+#include "MagicLeapCustomPresentVk.h"
+#include "MagicLeapHMD.h"
+#include "RenderingThread.h"
+#include "Lumin/CAPIShims/LuminAPI.h"
+
+#include "MagicLeapGraphics.h"
+
+#if PLATFORM_WINDOWS || PLATFORM_LUMIN
+#include "XRThreadUtils.h"
+#include "MagicLeapHelperVulkan.h"
+#include "VulkanRHIBridge.h"
+#include "VulkanContext.h"
+#endif // PLATFORM_WINDOWS || PLATFORM_LUMIN
+
+#include "Containers/Union.h"
+
+#if PLATFORM_WINDOWS || PLATFORM_LUMIN
+
+FMagicLeapCustomPresentVulkan::FMagicLeapCustomPresentVulkan(FMagicLeapHMD* plugin)
+: FMagicLeapCustomPresent(plugin)
+, RenderTargetTexture(VK_NULL_HANDLE)
+, RenderTargetTextureAllocation(VK_NULL_HANDLE)
+, RenderTargetTextureAllocationOffset(0)
+, RenderTargetTextureSRGB(VK_NULL_HANDLE)
+, LastAliasedRenderTarget(VK_NULL_HANDLE)
+{}
+
+void FMagicLeapCustomPresentVulkan::BeginRendering()
+{
+#if WITH_MLSDK
+	check(IsInRenderingThread());
+
+	ExecuteOnRHIThread([this]()
+	{
+		// Always use RHITrackingFrame here, which is then copied to the RenderTrackingFrame. 
+		FTrackingFrame& RHIframe = Plugin->RHITrackingFrame;
+		RHIframe.ProjectionType = MLGraphicsProjectionType_UnsignedZ;
+		BeginFrame(RHIframe);
+		if (bCustomPresentIsSet)
+		{
+			Plugin->InitializeRenderFrameFromRHIFrame();
+		}
+	});
+#endif // WITH_MLSDK
+}
+
+void FMagicLeapCustomPresentVulkan::FinishRendering()
+{
+#if WITH_MLSDK
+	check(IsInRenderingThread() || IsInRHIThread());
+
+	if (Plugin->IsDeviceInitialized() && Plugin->GetCurrentFrame().bBeginFrameSucceeded)
+	{
+		// Notify first render here instead of in RenderToMLSurfaces_RenderThread() because the render to MLSurfaces is already finished by now.
+		NotifyFirstRender();
+
+		// TODO [Blake] : Hack since we cannot yet specify a handle per view in the view family
+		const MLGraphicsVirtualCameraInfoArray& vp_array = Plugin->GetCurrentFrame().FrameInfo.virtual_camera_info_array;
+		const uint32 vp_width = static_cast<uint32>(vp_array.viewport.w);
+		const uint32 vp_height = static_cast<uint32>(vp_array.viewport.h);
+
+		// Clear MLGraphics eye textures since on Vulkan they are not pre-cleared. (needs to be done on RHI thread)
+		// Temporary. We are getting an API change and an OS update to clear them before sending to the app, same as in OpenGL.
+		FMagicLeapHelperVulkan::ClearImage(vp_array.color_id, FLinearColor::Transparent, 0 /* BaseMipLevel */, 1 /* LevelCount */, 0 /* BaseArrayLayer */, vp_array.num_virtual_cameras);
+
+		// Alias the render target with an srgb image description for proper color space output.
+		if (RenderTargetTextureAllocation != VK_NULL_HANDLE && LastAliasedRenderTarget != RenderTargetTexture)
+		{
+			// SDKUNREAL-1135: ML remote image is corrupted on AMD hardware
+			if (!IsRHIDeviceAMD())
+			{
+				FMagicLeapHelperVulkan::DestroyImageSRGB(RenderTargetTextureSRGB);
+				RenderTargetTextureSRGB = VK_NULL_HANDLE;
+				RenderTargetTextureSRGB = reinterpret_cast<void *>(FMagicLeapHelperVulkan::AliasImageSRGB((uint64)RenderTargetTextureAllocation, RenderTargetTextureAllocationOffset, vp_width * 2, vp_height));
+				check(RenderTargetTextureSRGB != VK_NULL_HANDLE);
+			}
+			LastAliasedRenderTarget = RenderTargetTexture;
+			UE_LOG(LogMagicLeap, Log, TEXT("Aliased render target for correct sRGB ouput."));
+		}
+
+		const VkImage FinalTarget = (RenderTargetTextureSRGB != VK_NULL_HANDLE) ? reinterpret_cast<const VkImage>(RenderTargetTextureSRGB) : reinterpret_cast<const VkImage>(RenderTargetTexture);
+		FMagicLeapHelperVulkan::BlitImage((uint64)FinalTarget, 0, 0, 0, vp_width, vp_height, 1, (uint64)vp_array.color_id, 0, 0, 0, 0, vp_width, vp_height, 1);
+		FMagicLeapHelperVulkan::BlitImage((uint64)FinalTarget, vp_width, 0, 0, vp_width, vp_height, 1, (uint64)vp_array.color_id, 1, 0, 0, 0, vp_width, vp_height, 1);
+
+		FMagicLeapHelperVulkan::SignalObjects((uint64)vp_array.virtual_cameras[0].sync_object, (uint64)vp_array.virtual_cameras[1].sync_object);
+
+		MLResult Result = MLGraphicsEndFrame(Plugin->GraphicsClient, Plugin->GetCurrentFrame().FrameInfo.handle);
+		if (Result != MLResult_Ok)
+		{
+#if !WITH_EDITOR
+			UE_LOG(LogMagicLeap, Error, TEXT("MLGraphicsEndFrame failed with status %d"), Result);
+#endif
+		}
+	}
+  
+	Plugin->InitializeOldFrameFromRenderFrame();
+#endif // WITH_MLSDK
+}
+
+void FMagicLeapCustomPresentVulkan::UpdateViewport(const FViewport& Viewport, FRHIViewport* InViewportRHI)
+{
+	check(IsInGameThread());
+	check(InViewportRHI);
+
+	const FTexture2DRHIRef& RT = Viewport.GetRenderTargetTexture();
+	check(IsValidRef(RT));
+
+	RenderTargetTexture = RT->GetNativeResource();
+	RenderTargetTextureAllocation = reinterpret_cast<void*>(static_cast<FVulkanTexture2D*>(RT->GetTexture2D())->Surface.GetAllocationHandle());
+	RenderTargetTextureAllocationOffset = static_cast<FVulkanTexture2D*>(RT->GetTexture2D())->Surface.GetAllocationOffset();
+
+	InViewportRHI->SetCustomPresent(this);
+
+	FMagicLeapCustomPresentVulkan* CustomPresent = this;
+	ENQUEUE_RENDER_COMMAND(UpdateViewport_RT)(
+		[CustomPresent](FRHICommandList& RHICmdList)
+		{
+			CustomPresent->UpdateViewport_RenderThread();
+		}
+	);
+}
+
+void FMagicLeapCustomPresentVulkan::UpdateViewport_RenderThread()
+{
+	check(IsInRenderingThread());
+
+	ExecuteOnRHIThread_DoNotWait([this]()
+	{
+		bCustomPresentIsSet = true;
+	});
+}
+
+void FMagicLeapCustomPresentVulkan::Reset()
+{
+	FMagicLeapCustomPresent::Reset();
+	if (IsInRenderingThread() && DestTextureRef.IsValid())
+	{
+		if (DestTextureRef.IsValid())
+		{
+			DestTextureRef.SafeRelease();
+		}
+	}
+
+#if WITH_MLSDK
+	FMagicLeapHelperVulkan::DestroyImageSRGB(RenderTargetTextureSRGB);
+	RenderTargetTextureSRGB = VK_NULL_HANDLE;
+#endif
+}
+
+#endif // PLATFORM_WINDOWS || PLATFORM_LUMIN
diff --git a/Engine/Plugins/Lumin/MagicLeap/Source/MagicLeap/Private/MagicLeapCustomPresentVk.h b/Engine/Plugins/Lumin/MagicLeap/Source/MagicLeap/Private/MagicLeapCustomPresentVk.h
new file mode 100644
--- /dev/null
+++ b/Engine/Plugins/Lumin/MagicLeap/Source/MagicLeap/Private/MagicLeapCustomPresentVk.h
@@ -0,0 +1,30 @@
+// Copyright 1998-2019 Epic Games, Inc. All Rights Reserved.
+
+#pragma once
+
+#include "MagicLeapGraphics.h"
+#include "MagicLeapCustomPresent.h"
+
+#if PLATFORM_WINDOWS || PLATFORM_LUMIN
+class FMagicLeapCustomPresentVulkan : public FMagicLeapCustomPresent
+{
+public:
+	FMagicLeapCustomPresentVulkan(FMagicLeapHMD* plugin);
+
+	virtual void BeginRendering() override;
+	virtual void FinishRendering() override;
+	virtual void UpdateViewport(const FViewport& Viewport, FRHIViewport* InViewportRHI) override;
+	virtual void UpdateViewport_RenderThread() override;
+	virtual void Reset() override;
+
+protected:
+	// Used as a container for the MLGraphics texture array.
+	FTexture2DArrayRHIRef DestTextureRef;
+
+	void* RenderTargetTexture;
+	void* RenderTargetTextureAllocation;
+	uint64 RenderTargetTextureAllocationOffset = 0;
+	void* RenderTargetTextureSRGB;
+	void* LastAliasedRenderTarget;
+};
+#endif // PLATFORM_WINDOWS || PLATFORM_LUMIN
diff --git a/Engine/Plugins/Lumin/MagicLeap/Source/MagicLeap/Private/MagicLeapGraphics.h b/Engine/Plugins/Lumin/MagicLeap/Source/MagicLeap/Private/MagicLeapGraphics.h
--- a/Engine/Plugins/Lumin/MagicLeap/Source/MagicLeap/Private/MagicLeapGraphics.h
+++ b/Engine/Plugins/Lumin/MagicLeap/Source/MagicLeap/Private/MagicLeapGraphics.h
@@ -1,54 +1,6 @@
-// %BANNER_BEGIN%
-// ---------------------------------------------------------------------
-// %COPYRIGHT_BEGIN%
-//
-// Copyright (c) 2017 Magic Leap, Inc. (COMPANY) All Rights Reserved.
-// Magic Leap, Inc. Confidential and Proprietary
-//
-// NOTICE: All information contained herein is, and remains the property
-// of COMPANY. The intellectual and technical concepts contained herein
-// are proprietary to COMPANY and may be covered by U.S. and Foreign
-// Patents, patents in process, and are protected by trade secret or
-// copyright law. Dissemination of this information or reproduction of
-// this material is strictly forbidden unless prior written permission is
-// obtained from COMPANY. Access to the source code contained herein is
-// hereby forbidden to anyone except current COMPANY employees, managers
-// or contractors who have executed Confidentiality and Non-disclosure
-// agreements explicitly covering such access.
-//
-// The copyright notice above does not evidence any actual or intended
-// publication or disclosure of this source code, which includes
-// information that is confidential and/or proprietary, and is a trade
-// secret, of COMPANY. ANY REPRODUCTION, MODIFICATION, DISTRIBUTION,
-// PUBLIC PERFORMANCE, OR PUBLIC DISPLAY OF OR THROUGH USE OF THIS
-// SOURCE CODE WITHOUT THE EXPRESS WRITTEN CONSENT OF COMPANY IS
-// STRICTLY PROHIBITED, AND IN VIOLATION OF APPLICABLE LAWS AND
-// INTERNATIONAL TREATIES. THE RECEIPT OR POSSESSION OF THIS SOURCE
-// CODE AND/OR RELATED INFORMATION DOES NOT CONVEY OR IMPLY ANY RIGHTS
-// TO REPRODUCE, DISCLOSE OR DISTRIBUTE ITS CONTENTS, OR TO MANUFACTURE,
-// USE, OR SELL ANYTHING THAT IT MAY DESCRIBE, IN WHOLE OR IN PART.
-//
-// %COPYRIGHT_END%
-// --------------------------------------------------------------------
-// %BANNER_END%
+// Copyright 1998-2019 Epic Games, Inc. All Rights Reserved.
 
 #pragma once
 
 #include "CoreTypes.h"
-
-
-#if PLATFORM_LUMIN
-// In order for ml_graphics.h to define Vulkan functions we need to include
-// the vulkan headers that define the VK_VERSION_1_0 macro.
-#include <vulkan.h>
-#endif // PLATFORM_LUMIN
-
-#include "MagicLeapPluginUtil.h" // for ML_INCLUDES_START/END
-
-#if WITH_MLSDK
-ML_INCLUDES_START
-#include <ml_graphics.h>
-ML_INCLUDES_END
-
-
-#endif //WITH_MLSDK
+#include "Lumin/CAPIShims/LuminAPIGraphics.h"
