Copying //UE4/Partner-MagicLeap-Staging to Dev-VR (//UE4/Dev-VR)

diff --git a/Engine/Plugins/Lumin/MagicLeap/Source/MagicLeap/Private/MagicLeapHMD.cpp b/Engine/Plugins/Lumin/MagicLeap/Source/MagicLeap/Private/MagicLeapHMD.cpp
--- a/Engine/Plugins/Lumin/MagicLeap/Source/MagicLeap/Private/MagicLeapHMD.cpp
+++ b/Engine/Plugins/Lumin/MagicLeap/Source/MagicLeap/Private/MagicLeapHMD.cpp
@@ -18,7 +18,11 @@
 #include "DeviceProfiles/DeviceProfile.h"
 #include "LuminRuntimeSettings.h"
 #include "MotionControllerComponent.h"
-#include "IMagicLeapInputDevice.h"
+#include "IMagicLeapTrackerEntity.h"
+#include "MagicLeapCFUID.h"
+#if PLATFORM_WINDOWS || PLATFORM_LINUX || PLATFORM_MAC
+#include "Async/Async.h"
+#endif // PLATFORM_WINDOWS || PLATFORM_LINUX || PLATFORM_MAC
 
 #include "XRThreadUtils.h"
 
@@ -31,38 +35,26 @@
 #include "UnrealEngine.h"
 #include "ClearQuad.h"
 #include "MagicLeapXRCamera.h"
-#include "IMagicLeapInputDevice.h"
 #include "GeneralProjectSettings.h"
 #include "MagicLeapSettings.h"
-#include "MagicLeapSDKDetection.h"
-#include "IMagicLeapModule.h"
 
-#if !PLATFORM_MAC // @todo Lumin: I had to add this to get Mac to compile - trying to add GL to Mac build had massive compile issues
+#if !PLATFORM_MAC
 #include "OpenGLDrv.h"
 #include "VulkanRHIPrivate.h"
 #include "VulkanRHIBridge.h"
-#endif
-
-#include "MagicLeapPluginUtil.h" // for ML_INCLUDES_START/END
+#else
+#include "MetalRHI.h"
+#endif // !PLATFORM_MAC
 
-#if WITH_MLSDK
-ML_INCLUDES_START
-#include <ml_perception.h>
+#include "Lumin/CAPIShims/LuminAPIPerception.h"
+#include "Lumin/CAPIShims/LuminAPIRemote.h"
+#include "Lumin/CAPIShims/LuminAPIGraphics.h"
+#include "Lumin/CAPIShims/LuminAPIPrivileges.h"
 
 #if !PLATFORM_LUMIN
-#include <ml_remote.h>
 #include "Misc/MessageDialog.h"
 #endif
 
-#include <ml_graphics.h>
-
-#include "ml_privileges.h"
-
-#include "ml_privileges.h"
-
-ML_INCLUDES_END
-#endif //WITH_MLSDK
-
 #if WITH_EDITOR
 #include "Editor/EditorEngine.h"
 #include "Editor/EditorPerformanceSettings.h"
@@ -85,56 +77,54 @@ public:
 	FMagicLeapPlugin()
 		: bIsVDZIEnabled(false)
 		, bUseVulkanForZI(false)
+		, bUseMLAudioForZI(false)
 	{
 	}
 
 	/** IModuleInterface implementation */
 	virtual void StartupModule() override
 	{
-		FMagicLeapSDKDetection::DetectSDK();
-
 		// Ideally, we should be able to query GetDefault<UMagicLeapSettings>()->bEnableZI directly.
 		// Unfortunately, the UObject system hasn't finished initialization when this module has been loaded.
 		GConfig->GetBool(TEXT("/Script/MagicLeap.MagicLeapSettings"), TEXT("bEnableZI"), bIsVDZIEnabled, GEngineIni);
 		GConfig->GetBool(TEXT("/Script/MagicLeap.MagicLeapSettings"), TEXT("bUseVulkanForZI"), bUseVulkanForZI, GEngineIni);
+		GConfig->GetBool(TEXT("/Script/MagicLeap.MagicLeapSettings"), TEXT("bUseMLAudioForZI"), bUseMLAudioForZI, GEngineIni);
 
-		APISetup.Startup(bIsVDZIEnabled);
-#if WITH_MLSDK
-		APISetup.LoadDLL(TEXT("ml_perception_client"));
-		APISetup.LoadDLL(TEXT("ml_graphics"));
-		APISetup.LoadDLL(TEXT("ml_lifecycle"));
-		APISetup.LoadDLL(TEXT("ml_privileges"));
-#endif
-
+#if (PLATFORM_WINDOWS || PLATFORM_MAC) && WITH_EDITOR
 		if (bIsVDZIEnabled)
 		{
-#if PLATFORM_WINDOWS
-
-			const bool bLoadedRemoteDLL = APISetup.LoadDLL(TEXT("ml_remote"));
+#if WITH_MLSDK
+			const bool bLoadedRemoteDLL = (MLSDK_API::LibraryLoader::Ref().LoadDLL(TEXT("ml_remote")) != nullptr);
+#else
+			const bool bLoadedRemoteDLL = false;
+#endif // WITH_MLSDK
 			if (!bLoadedRemoteDLL)
 			{
 				// Bail early, because we'll eventually die later.
 				UE_LOG(LogMagicLeap, Warning, TEXT("VDZI enabled, but can't load the ml_remote DLL.  Is your MLSDK directory set up properly?"));
 				bIsVDZIEnabled = false;
 			}
-			
+
 			FString CommandLine = FCommandLine::Get();
+#if !PLATFORM_MAC
 			const FString GLFlag(" -opengl4 ");
 			const FString VKFlag(" -vulkan ");
+			const FString GLToken = GLFlag.TrimStartAndEnd();
+			const FString VKToken = VKFlag.TrimStartAndEnd();
+
+			CommandLine.ReplaceInline(*(GLToken.TrimEnd()), TEXT(""));
+			CommandLine.ReplaceInline(*(VKToken.TrimEnd()), TEXT(""));
+#endif // !PLATFORM_MAC
 
+			const FString AudioMixerFlag(" -audiomixer MagicLeapAudio ");
+			const FString AudioMixerToken = AudioMixerFlag.TrimStartAndEnd();
+			CommandLine.ReplaceInline(*(AudioMixerToken.TrimEnd()), TEXT(""));
+
+#if !PLATFORM_MAC
 			if (bUseVulkanForZI)
 			{
 				UE_LOG(LogMagicLeap, Log, TEXT("ML VDZI mode enabled. Using Vulkan renderer."));
-				int32 GLFlagOffset = CommandLine.Find(GLFlag);
-				if (GLFlagOffset != INDEX_NONE) CommandLine.RemoveAt(GLFlagOffset, GLFlag.Len());
-				if (CommandLine.Find(VKFlag) == INDEX_NONE) CommandLine.Append(VKFlag);
-
-				// r.Vulkan.RHIThread=0 is requried for Vulkan on Windows with MLRemote. Setting it in BeginPlay() doesnt help.
-				IConsoleVariable* CVar = IConsoleManager::Get().FindConsoleVariable(TEXT("r.Vulkan.RHIThread"));
-				if (CVar)
-				{
-					// CVar->Set(TEXT("0"));
-				}
+				ExtraCommandLineFlags = VKFlag;
 			}
 			else
 			{
@@ -142,15 +132,29 @@ public:
 				// OpenGL is forced by loading module in PostConfigInit phase and passing in command line.
 				// -opengl will force editor to use OpenGL3/SM4 feature level. Fwd VR path requires SM5 feature level, thus passing -opengl here will break editor preview window with Fwd VR path
 				// The cmd arg for OpenGL4/SM5 feature level is -opengl4 in Windows.
-				UE_LOG(LogMagicLeap, Log, TEXT("ML VDZI mode enabled. Using OpenGL renderer."));
-				int32 VKFlagOffset = CommandLine.Find(VKFlag);
-				if (VKFlagOffset != INDEX_NONE) CommandLine.RemoveAt(VKFlagOffset, VKFlag.Len());
-				if (CommandLine.Find(GLFlag) == INDEX_NONE) CommandLine.Append(GLFlag);
+				ExtraCommandLineFlags = GLFlag;
+			}
+#endif // !PLATFORM_MAC
+
+			if (bUseMLAudioForZI)
+			{
+				// Cache the audio mixer module already set for the project.
+				GConfig->GetString(TEXT("Audio"), TEXT("AudioMixerModuleName"), PreviousAudioMixer, GEngineIni);
+
+				GConfig->SetString(TEXT("Audio"), TEXT("AudioMixerModuleName"), TEXT("MagicLeapAudio"), GEngineIni);
+				ExtraCommandLineFlags.Append(AudioMixerFlag);
 			}
 
+			CommandLine.Append(ExtraCommandLineFlags);
 			FCommandLine::Set(*CommandLine);
-#endif // PLATFORM_WINDOWS
+
+			// This is the closest callback that happens before EditorExit. We need to undo our
+			// command line changes before the EditorExit as that where the editor will relaunch
+			// itself using the current command line. Hence to avoid the RHI option stick when
+			// we change the VDZI option we need to remove before the EngineExit.
+			FEditorDelegates::OnShutdownPostPackagesSaved.AddRaw(this, &FMagicLeapPlugin::UndoCommandLine);
 		}
+#endif // (PLATFORM_WINDOWS || PLATFORM_MAC) && WITH_EDITOR
 
 #if WITH_EDITOR
 		// We don't quite have control of when the "Settings" module is loaded, so we'll wait until PostEngineInit to register settings.
@@ -160,12 +164,28 @@ public:
 		IMagicLeapPlugin::StartupModule();
 	}
 
+#if WITH_EDITOR
+	void UndoCommandLine()
+	{
+		if (ExtraCommandLineFlags.Len() > 0)
+		{
+			FString CommandLine = FCommandLine::Get();
+			CommandLine.ReplaceInline(*(ExtraCommandLineFlags.TrimStartAndEnd()), TEXT(""));
+			FCommandLine::Set(*CommandLine);
+		}
+
+		if (bUseMLAudioForZI)
+		{
+			GConfig->SetString(TEXT("Audio"), TEXT("AudioMixerModuleName"), *PreviousAudioMixer, GEngineIni);
+		}
+	}
+#endif
+
 	virtual void ShutdownModule() override
 	{
 #if WITH_EDITOR
 		RemoveEditorSettings();
 #endif
-		APISetup.Shutdown();
 		IMagicLeapPlugin::ShutdownModule();
 	}
 
@@ -182,7 +202,6 @@ public:
 		
 		TSharedPtr<FMagicLeapHMD, ESPMode::ThreadSafe> LocalHMD;
 
-#if !PLATFORM_MAC
 		if (HMD.IsValid())
 		{
 			LocalHMD = HMD.Pin();
@@ -198,13 +217,10 @@ public:
 			HMD = LocalHMD;
 			ARModule->ConnectARImplementationToXRSystem(LocalHMD.Get());
 			ARModule->InitializeARImplementation();
-		}
-#endif // WITH_MLSDK
-#endif // !PLATFORM_MAC
-		if (LocalHMD.IsValid() && !LocalHMD->IsInitialized())
-		{
 			LocalHMD->Startup();
 		}
+#endif // WITH_MLSDK
+
 		return LocalHMD;
 	}
 
@@ -228,22 +244,6 @@ public:
 #endif // PLATFORM_LUMIN
 	}
 
-	virtual TWeakPtr<IMagicLeapHMD, ESPMode::ThreadSafe> GetHMD() override
-	{
-#if WITH_EDITOR
-		if (bIsVDZIEnabled)
-		{
-			return HMD;
-		}
-		else
-		{
-			return nullptr;
-		}
-#else
-		return HMD;
-#endif
-	}
-
 	virtual TSharedPtr<IHeadMountedDisplayVulkanExtensions, ESPMode::ThreadSafe> GetVulkanExtensions() override
 	{
 #if !PLATFORM_LUMIN
@@ -264,43 +264,103 @@ public:
 		return nullptr;
 	}
 
-	virtual void RegisterMagicLeapInputDevice(IMagicLeapInputDevice* InputDevice) override
+	virtual void RegisterMagicLeapTrackerEntity(IMagicLeapTrackerEntity* TrackerEntity) override
 	{
-		check(InputDevice);
-		InputDevices.Add(InputDevice);
+		check(TrackerEntity);
+		TrackerEntities.Add(TrackerEntity);
 	}
 
-	virtual void UnregisterMagicLeapInputDevice(IMagicLeapInputDevice* InputDevice) override
+	virtual void UnregisterMagicLeapTrackerEntity(IMagicLeapTrackerEntity* TrackerEntity) override
 	{
-		check(InputDevice);
-		InputDevices.Remove(InputDevice);
+		check(TrackerEntity);
+		TrackerEntities.Remove(TrackerEntity);
 	}
 
-	virtual void EnableInputDevices() override
+	void EnableTrackerEntities()
 	{
-		for (auto& It : InputDevices)
+		for (auto& It : TrackerEntities)
 		{
-			if (It->SupportsExplicitEnable())
-			{
-				It->Enable();
-			}
+			It->CreateEntityTracker();
+		}
+	}
+
+	void DisableTrackerEntities()
+	{
+		for (auto& It : TrackerEntities)
+		{
+			It->DestroyEntityTracker();
+		}
+	}
+
+	virtual bool SetIgnoreInput(bool Ignore) override
+	{
+		if (HMD.IsValid())
+		{
+			return HMD.Pin()->SetIgnoreInput(Ignore);
+		}
+		return false;
+	}
+
+	void OnBeginRendering_GameThread()
+	{
+		for (auto& It : TrackerEntities)
+		{
+			It->OnBeginRendering_GameThread();
+		}
+	}
+
+	virtual bool GetTransform(const FGuid& Id, FTransform& OutTransform, EMagicLeapTransformFailReason& OutReason) const override
+	{
+#if WITH_MLSDK
+		if (HMD.IsValid())
+		{
+			MLCoordinateFrameUID CFUID;
+			MagicLeap::FGuidToMLCFUID(Id, CFUID);
+			return HMD.Pin()->GetAppFrameworkConst().GetTransform(CFUID, OutTransform, OutReason);
 		}
+#endif // WITH_MLSDK
+
+		OutReason = EMagicLeapTransformFailReason::HMDNotInitialized;
+		return false;
 	}
 
-	virtual void DisableInputDevices() override
+#if WITH_MLSDK
+	virtual bool GetTransform(const MLCoordinateFrameUID& Id, FTransform& OutTransform, EMagicLeapTransformFailReason& OutReason) const override
 	{
-		for (auto& It : InputDevices)
+		if (HMD.IsValid())
 		{
-			It->Disable();
+			return HMD.Pin()->GetAppFrameworkConst().GetTransform(Id, OutTransform, OutReason);
 		}
+
+		OutReason = EMagicLeapTransformFailReason::HMDNotInitialized;
+		return false;
 	}
+#endif // WITH_MLSDK
 
-	virtual void OnBeginRendering_GameThread_UpdateInputDevices() override
+	virtual float GetWorldToMetersScale() const override
 	{
-		for (auto& It : InputDevices)
+		if (HMD.IsValid())
 		{
-			It->OnBeginRendering_GameThread_Update();
+			return HMD.Pin()->GetAppFrameworkConst().GetWorldToMetersScale();
 		}
+
+		return 100.0f;
+	}
+
+	virtual bool IsPerceptionEnabled() const override
+	{
+		if (HMD.IsValid())
+		{
+			return HMD.Pin()->IsPerceptionEnabled();
+		}
+
+		return false;
+	}
+
+	virtual bool UseMLAudioForZI() const override
+	{
+		// if VDZI is disabled, doesn't matter if MLAudio option is enabled.
+		return bIsVDZIEnabled && bUseMLAudioForZI;
 	}
 
 private:
@@ -335,11 +395,13 @@ private:
 
 	bool bIsVDZIEnabled;
 	bool bUseVulkanForZI;
-	FMagicLeapAPISetup APISetup;
+	bool bUseMLAudioForZI;
 	TWeakPtr<FMagicLeapHMD, ESPMode::ThreadSafe> HMD;
 	TSharedPtr<IARSystemSupport, ESPMode::ThreadSafe> ARImplementation;
 	TSharedPtr<FMagicLeapVulkanExtensions, ESPMode::ThreadSafe> VulkanExtensions;
-	TSet<IMagicLeapInputDevice*> InputDevices;
+	TSet<IMagicLeapTrackerEntity*> TrackerEntities;
+	FString ExtraCommandLineFlags;
+	FString PreviousAudioMixer;
 };
 
 IMPLEMENT_MODULE(FMagicLeapPlugin, MagicLeap)
@@ -347,6 +409,7 @@ IMPLEMENT_MODULE(FMagicLeapPlugin, MagicLeap)
 //////////////////////////////////////////////////////////////////////////
 
 const FName FMagicLeapHMD::SystemName(TEXT("MagicLeap"));
+const float FMagicLeapHMD::DefaultFocusDistance = 500.0f;
 
 FName FMagicLeapHMD::GetSystemName() const
 {
@@ -426,34 +489,66 @@ bool FMagicLeapHMD::OnEndGameFrame(FWorldContext& WorldContext)
 
 bool FMagicLeapHMD::IsHMDConnected()
 {
-#if WITH_MLSDK
 #if PLATFORM_LUMIN
-	return AppFramework.IsInitialized();
-#elif PLATFORM_WINDOWS
-	bool bZIServerRunning = false;
-	if (bIsVDZIEnabled)
+	return true;
+#elif WITH_MLSDK && (PLATFORM_WINDOWS || PLATFORM_LINUX || PLATFORM_MAC)
+	const bool bLoadedRemoteDLL = (MLSDK_API::LibraryLoader::Ref().LoadDLL(TEXT("ml_remote")) != nullptr);
+	if (!bIsVDZIEnabled || !bLoadedRemoteDLL)
 	{
-		if (FMagicLeapSDKDetection::IsSDKDetected())
+		return false;
+	}
+
+	EServerPingState ServerPingState = static_cast<EServerPingState>(FPlatformAtomics::AtomicRead(&ZIServerPingState));
+	switch (ServerPingState)
+	{
+	case EServerPingState::NotInitialized:
+	{
+		FPlatformAtomics::InterlockedExchange(&ZIServerPingState, static_cast<int32>(EServerPingState::AttemptingFirstPing));
+		AsyncTask(ENamedThreads::AnyBackgroundThreadNormalTask, [this]()
 		{
-			MLResult Result = MLRemoteIsServerConfigured(&bZIServerRunning);
-			if (Result != MLResult_Ok)
+			while (!IsEngineExitRequested())
 			{
-				UE_LOG(LogMagicLeap, Error, TEXT("MLRemoteIsServerConfigured failed with error %s!"), UTF8_TO_TCHAR(MLGetResultString(Result)));
+				bool bIsZIServerRunning;
+				MLResult Result = MLRemoteIsServerConfigured(&bIsZIServerRunning);
+				if (Result != MLResult_Ok && Result != MLResult_Timeout)
+				{
+					UE_LOG(LogMagicLeap, Error, TEXT("MLRemoteIsServerConfigured failed with error %s!"), UTF8_TO_TCHAR(MLGetResultString(Result)));
+				}
 
-				// Ensure we don't falsely mark it as running if there was an error
-				bZIServerRunning = false;
+				if (bIsZIServerRunning)
+				{
+					FPlatformAtomics::InterlockedExchange(&ZIServerPingState, static_cast<int32>(EServerPingState::ServerOnline));
+				}
+				else
+				{
+					FPlatformAtomics::InterlockedExchange(&ZIServerPingState, static_cast<int32>(EServerPingState::ServerOffline));
+				}
+
+				FPlatformProcess::Sleep(ServerPingInterval);
 			}
-		}
+		});
 	}
-	// TODO: MLVirtualDeviceIsServerRunning() crashes when called on render thread.
-	return AppFramework.IsInitialized() && bIsVDZIEnabled && bZIServerRunning;
-#else
-	return false;
-#endif
+	break;
 
-#else
+	case EServerPingState::AttemptingFirstPing:
+	{
+		do 
+		{
+			ServerPingState = static_cast<EServerPingState>(FPlatformAtomics::AtomicRead(&ZIServerPingState));
+		} while (ServerPingState == EServerPingState::AttemptingFirstPing);
+	}
+	break;
+
+	case EServerPingState::ServerOnline:
+	case EServerPingState::ServerOffline:
+	{
+	}
+	break;
+	}
+
+	return ServerPingState == EServerPingState::ServerOnline;
+#endif // PLATFORM_LUMIN
 	return false;
-#endif //WITH_MLSDK
 }
 
 bool FMagicLeapHMD::IsHMDEnabled() const
@@ -475,7 +570,7 @@ bool FMagicLeapHMD::GetHMDMonitorInfo(MonitorInfo& MonitorDesc)
 	// Use ML device resolution only when HMD is enabled.
 	// This ensures that we don't provide an invalid resolution when the device is not connected.
 	// TODO: check if we can rely on the return bool value from AppFramework.GetDeviceResolution() instead.
-	if (IsInitialized() && bHmdEnabled)
+	if (bHmdEnabled)
 	{
 		const FIntPoint& RTSize = GetIdealRenderTargetSize();
 		MonitorDesc.MonitorName = "";
@@ -500,9 +595,11 @@ bool FMagicLeapHMD::GetHMDMonitorInfo(MonitorInfo& MonitorDesc)
 
 void FMagicLeapHMD::GetFieldOfView(float& OutHFOVInDegrees, float& OutVFOVInDegrees) const
 {
-	const FTrackingFrame& frame = GetCurrentFrame();
-	OutHFOVInDegrees = frame.HFov;
-	OutVFOVInDegrees = frame.VFov;
+	FMagicLeapCustomPresent* const CP = GetActiveCustomPresent(true);
+	if (CP)
+	{
+		CP->GetFieldOfView(OutHFOVInDegrees, OutVFOVInDegrees);
+	}
 }
 
 void FMagicLeapHMD::SetPixelDensity(const float NewDensity)
@@ -513,7 +610,7 @@ void FMagicLeapHMD::SetPixelDensity(const float NewDensity)
 
 FIntPoint FMagicLeapHMD::GetIdealRenderTargetSize() const
 {
-	return FIntPoint(1280 * 2, 960);
+	return DefaultRenderTargetSize;
 }
 
 bool FMagicLeapHMD::DoesSupportPositionalTracking() const
@@ -579,7 +676,7 @@ bool FMagicLeapHMD::IsHeadTrackingAllowed() const
 void FMagicLeapHMD::ResetOrientationAndPosition(float yaw)
 {
 	ResetOrientation(yaw);
-	ResetPosition();
+	// TODO: we don't reset position? Tecnically none of this should be required. Everything related to this topic should be handled by XRPawn.
 }
 
 void FMagicLeapHMD::ResetOrientation(float Yaw)
@@ -596,44 +693,6 @@ void FMagicLeapHMD::ResetOrientation(float Yaw)
 		ViewRotation.Yaw -= Yaw;
 		ViewRotation.Normalize();
 	}
-
-	AppFramework.SetBaseOrientation(ViewRotation.Quaternion());
-}
-
-void FMagicLeapHMD::ResetPosition()
-{
-	const FTrackingFrame& frame = GetCurrentFrame();
-	AppFramework.SetBasePosition(frame.RawPose.GetTranslation());
-}
-
-void FMagicLeapHMD::SetBasePosition(const FVector& InBasePosition)
-{
-	AppFramework.SetBasePosition(InBasePosition);
-}
-
-FVector FMagicLeapHMD::GetBasePosition() const
-{
-	return AppFramework.GetBasePosition();
-}
-
-void FMagicLeapHMD::SetBaseRotation(const FRotator& BaseRot)
-{
-	AppFramework.SetBaseRotation(BaseRot);
-}
-
-FRotator FMagicLeapHMD::GetBaseRotation() const
-{
-	return AppFramework.GetBaseRotation();
-}
-
-void FMagicLeapHMD::SetBaseOrientation(const FQuat& BaseOrient)
-{
-	AppFramework.SetBaseOrientation(BaseOrient);
-}
-
-FQuat FMagicLeapHMD::GetBaseOrientation() const
-{
-	return AppFramework.GetBaseOrientation();
 }
 
 bool FMagicLeapHMD::EnumerateTrackedDevices(TArray<int32>& OutDevices, EXRTrackedDeviceType Type /*= EXRTrackedDeviceType::Any*/)
@@ -661,8 +720,11 @@ void FMagicLeapHMD::RefreshTrackingFrame()
 	GameTrackingFrame.FrameNumber = GFrameCounter;
 
 	// set the horizontal and vertical fov for this frame
-	GameTrackingFrame.HFov = AppFramework.GetFieldOfView().X;
-	GameTrackingFrame.VFov = AppFramework.GetFieldOfView().Y;
+	FMagicLeapCustomPresent* const CP = GetActiveCustomPresent(true);
+	if (CP)
+	{
+		CP->GetFieldOfView(GameTrackingFrame.HFov, GameTrackingFrame.VFov);
+	}
 
 	// Release the snapshot of the previous frame.
 	// This is done here instead of on end frame because modules implemented as input devices (Gestures, controller)
@@ -687,7 +749,35 @@ void FMagicLeapHMD::RefreshTrackingFrame()
 		HeadTrackingState.Confidence = state.confidence;
 	}
 
-	EFailReason FailReason = EFailReason::None;
+	// MLHandle because Unreal's uint64 is 'unsigned long long' whereas uint64_t for the C-API is 'unsigned long'
+	MLHandle MapEventsFlags = 0;
+	Result = MLHeadTrackingGetMapEvents(HeadTracker, &MapEventsFlags);
+#if PLATFORM_LUMIN
+	UE_CLOG(Result != MLResult_Ok, LogMagicLeap, Error, TEXT("MLHeadTrackingGetMapEvents failed with error %s."), UTF8_TO_TCHAR(MLGetResultString(Result)));
+#endif //PLATFORM_LUMIN
+
+	bHeadposeMapEventsAvailable = (Result == MLResult_Ok);
+	if (bHeadposeMapEventsAvailable)
+	{
+		static TArray<MLHeadTrackingMapEvent> AllMLFlags({
+			MLHeadTrackingMapEvent_Lost,
+			MLHeadTrackingMapEvent_Recovered,
+			MLHeadTrackingMapEvent_RecoveryFailed,
+			MLHeadTrackingMapEvent_NewSession
+		});
+
+		HeadposeMapEvents.Empty();
+
+		for (MLHeadTrackingMapEvent flag : AllMLFlags)
+		{
+			if ((MapEventsFlags & static_cast<MLHeadTrackingMapEvent>(flag)) != 0)
+			{
+				HeadposeMapEvents.Add(MLToUnrealHeadTrackingMapEvent(flag));
+			}
+		}
+	}
+
+	EMagicLeapTransformFailReason FailReason = EMagicLeapTransformFailReason::None;
 	// get the raw pose and tracking status for the frame
 	FTransform head_transform;
 	GameTrackingFrame.HasHeadTrackingPosition = AppFramework.GetTransform(GameTrackingFrame.FrameId, head_transform, FailReason);
@@ -695,7 +785,7 @@ void FMagicLeapHMD::RefreshTrackingFrame()
 	{
 		GameTrackingFrame.RawPose = head_transform;
 	}
-	else if (FailReason == EFailReason::NaNsInTransform)
+	else if (FailReason == EMagicLeapTransformFailReason::NaNsInTransform)
 	{
 		UE_LOG(LogMagicLeap, Error, TEXT("NaNs in head transform."));
 		GameTrackingFrame.RawPose = OldTrackingFrame.RawPose;
@@ -715,24 +805,62 @@ void FMagicLeapHMD::RefreshTrackingFrame()
 
 	if (!FocusActor.IsValid())
 	{
-		UE_LOG(LogMagicLeap, Verbose, TEXT("Focus actor not set. Defaulting focus distance to 500.0 cm. Call the SetFocusActor() function to set a valid focus actor."));
+		if (StabilizationDepthActor.IsValid())
+		{
+			UE_LOG(LogMagicLeap, Verbose, TEXT("Focus actor not set. Defaulting depth to stabilization actor. Call the SetFocusActor() function to set a valid actor."));
+			FocusActor = StabilizationDepthActor;
+		}
+		else
+		{
+			UE_LOG(LogMagicLeap, Verbose, TEXT("Neither the Focus actor not the Stabilization actor are set. Defaulting focus distance to 500.0 cm. Call the SetFocusActor() function to set a valid focus actor."));
+		}
+	}
+
+	if (!StabilizationDepthActor.IsValid())
+	{
+		if (FocusActor.IsValid())
+		{
+			UE_LOG(LogMagicLeap, Verbose, TEXT("Stabilization depth actor not set. Defaulting depth to focus actor. Call the SetStabilizationDepthActor() function to set a valid actor."));
+			StabilizationDepthActor = FocusActor;
+		}
+		else
+		{
+			UE_LOG(LogMagicLeap, Verbose, TEXT("Neither the Focus actor not the Stabilization actor are set. Defaulting depth to %f. Call the SetStabilizationDepthActor() function to set a valid actor."), GameTrackingFrame.FarClippingPlane);
+		}
 	}
 
 	// If GNearClipPlane is changed by the app at runtime,
 	// ensure we clamp the near clip to the value provided by ml_graphics.
 	UpdateNearClippingPlane();
 
-	FVector focusPoint = (FocusActor.IsValid()) ? FocusActor->GetActorLocation() : (CurrentOrientation.GetForwardVector() * 500.0f + CurrentPosition);
-	float focusDistance = FVector::DotProduct(focusPoint - CurrentPosition, CurrentOrientation.GetForwardVector());
-	GameTrackingFrame.FocusDistance = (focusDistance > GNearClippingPlane) ? focusDistance : GNearClippingPlane;
+	const FTransform WorldToTracking = GetTrackingToWorldTransform().Inverse();
+
+	const FVector FocusPoint = (FocusActor.IsValid()) ? WorldToTracking.TransformPosition(FocusActor->GetActorLocation()) : ((CurrentOrientation.GetForwardVector() * FMagicLeapHMD::DefaultFocusDistance) + CurrentPosition);
+	GameTrackingFrame.FocusDistance = (FocusPoint - CurrentPosition).Size();
+	// focus distance is within the range of NearClippingPlane  <= FocusDistance <= FarClippingPlane
+	GameTrackingFrame.FocusDistance = FMath::Clamp(GameTrackingFrame.FocusDistance, GNearClippingPlane, GameTrackingFrame.FarClippingPlane);
+
+	// We default to the same value as far clipping plane when actor is not valid
+	const float StabilizationDepthDistance = (StabilizationDepthActor.IsValid()
+		? (WorldToTracking.TransformPosition(StabilizationDepthActor->GetActorLocation()) - CurrentPosition).Size()
+		: GameTrackingFrame.FarClippingPlane);
+
+	// If set focus distance is valid (greater than near clipping plane), then we must check to see if
+	// stabilization depth is within the range of FocusDistance <= StabilizationDepthDistance <= FarClippingPlane
+	GameTrackingFrame.StabilizationDepth = FMath::Clamp(StabilizationDepthDistance, GameTrackingFrame.FocusDistance, GameTrackingFrame.FarClippingPlane);
+
+	// NearClippingPlane <= FocusDistance <= StabilizationDepthDistance <= FarClippingPlane
+	check(GNearClippingPlane <= GameTrackingFrame.FocusDistance);
+	check(GameTrackingFrame.FocusDistance <= GameTrackingFrame.StabilizationDepth);
+	check(GameTrackingFrame.StabilizationDepth <= GameTrackingFrame.FarClippingPlane);
 #endif //WITH_MLSDK
 }
 
 #if WITH_MLSDK
 
-EHeadTrackingError FMagicLeapHMD::MLToUnrealHeadTrackingError(MLHeadTrackingError error) const
+EMagicLeapHeadTrackingError FMagicLeapHMD::MLToUnrealHeadTrackingError(MLHeadTrackingError error) const
 {
-	#define ERRCASE(x) case MLHeadTrackingError_##x: { return EHeadTrackingError::x; }
+	#define ERRCASE(x) case MLHeadTrackingError_##x: { return EMagicLeapHeadTrackingError::x; }
 	switch(error)
 	{
 		ERRCASE(None)
@@ -740,17 +868,34 @@ EHeadTrackingError FMagicLeapHMD::MLToUnrealHeadTrackingError(MLHeadTrackingErro
 		ERRCASE(LowLight)
 		ERRCASE(Unknown)
 	}
-	return EHeadTrackingError::Unknown;
+	return EMagicLeapHeadTrackingError::Unknown;
 }
 
-EHeadTrackingMode FMagicLeapHMD::MLToUnrealHeadTrackingMode(MLHeadTrackingMode mode) const
+EMagicLeapHeadTrackingMode FMagicLeapHMD::MLToUnrealHeadTrackingMode(MLHeadTrackingMode mode) const
 {
 	switch(mode)
 	{
-		case MLHeadTrackingMode_6DOF: { return EHeadTrackingMode::PositionAndOrientation; }
-		case MLHeadTrackingMode_3DOF: { return EHeadTrackingMode::OrientationOnly; }
+		case MLHeadTrackingMode_6DOF: { return EMagicLeapHeadTrackingMode::PositionAndOrientation; }
+		case MLHeadTrackingMode_Unavailable: { return EMagicLeapHeadTrackingMode::Unavailable; }
+	}
+	return EMagicLeapHeadTrackingMode::Unknown;
+}
+
+EMagicLeapHeadTrackingMapEvent FMagicLeapHMD::MLToUnrealHeadTrackingMapEvent(MLHeadTrackingMapEvent mapevent) const
+{
+	#define MAPEVENTCASE(x) case MLHeadTrackingMapEvent_##x: { return EMagicLeapHeadTrackingMapEvent::x; }
+	switch(mapevent)
+	{
+		MAPEVENTCASE(Lost)
+		MAPEVENTCASE(Recovered)
+		MAPEVENTCASE(RecoveryFailed)
+		MAPEVENTCASE(NewSession)
+	}
+
+	{
+		checkNoEntry()
+		return EMagicLeapHeadTrackingMapEvent::Lost;
 	}
-	return EHeadTrackingMode::Unknown;
 }
 #endif //WITH_MLSDK
 
@@ -790,55 +935,98 @@ void FMagicLeapHMD::DisplayWarningIfVDZINotEnabled()
 }
 #endif
 
-#if PLATFORM_LUMIN
-void FMagicLeapHMD::SetFrameTimingHint(ELuminFrameTimingHint InFrameTimingHint)
+void FMagicLeapHMD::GetClipExtents()
 {
-	const static UEnum* FrameTimingEnum = StaticEnum<ELuminFrameTimingHint>();
-	check(FrameTimingEnum != nullptr);
-
-	if (InFrameTimingHint != CurrentFrameTimingHint)
+#if WITH_MLSDK
+	if (GraphicsClient != ML_INVALID_HANDLE)
 	{
-		if (GraphicsClient != ML_INVALID_HANDLE)
+		// get the clip extents for clipping content in update thread
+		MLGraphicsClipExtentsParams ClipExtentsParams;
+		MLGraphicsClipExtentsParamsInit(&ClipExtentsParams);
+		MLResult Result = MLGraphicsGetClipExtentsEx(GraphicsClient, &ClipExtentsParams, &GameTrackingFrame.UpdateInfoArray);
+		if (Result != MLResult_Ok)
 		{
-			MLGraphicsFrameTimingHint FTH = MLGraphicsFrameTimingHint_Unspecified;
+			FString ErrorMesg = FString::Printf(TEXT("MLGraphicsGetClipExtentsEx failed with error %s"), UTF8_TO_TCHAR(MLGetResultString(Result)));
 
-			switch (InFrameTimingHint)
-			{
-			case ELuminFrameTimingHint::Unspecified:
-				FTH = MLGraphicsFrameTimingHint_Unspecified;
-				break;
-			case ELuminFrameTimingHint::Maximum:
-				FTH = MLGraphicsFrameTimingHint_Maximum;
-				break;
-			case ELuminFrameTimingHint::FPS_60:
-				FTH = MLGraphicsFrameTimingHint_60Hz;
-				break;
-			case ELuminFrameTimingHint::FPS_120:
-				FTH = MLGraphicsFrameTimingHint_120Hz;
-				break;
-			default:
-				FTH = MLGraphicsFrameTimingHint_Unspecified;
-				UE_LOG(LogMagicLeap, Warning, TEXT("Tried to set invalid Frame Timing Hint!  Defaulting to unspecified."));
-			}
+			// In case we're running under VD/ZI, there's always the risk of disconnects.
+			// In those cases, the graphics API can return an error, but the client handle might still be valid.
+			// So we need to ensure that we always have valid data to prevent any Nan-related errors.
+			// On Lumin, we'll just assert.
+#if PLATFORM_WINDOWS || PLATFORM_LINUX || PLATFORM_MAC
+			GameTrackingFrame.FrameInfo.handle = ML_INVALID_HANDLE;
+			MagicLeap::ResetClipExtentsInfoArray(GameTrackingFrame.UpdateInfoArray);
+			UE_LOG(LogMagicLeap, Error, TEXT("%s"), *ErrorMesg);
+#else
+			UE_LOG(LogMagicLeap, Fatal, TEXT("%s"), *ErrorMesg);
+#endif // PLATFORM_WINDOWS || PLATFORM_LINUX || PLATFORM_MAC
+		}
+	}
+#endif // WITH_MLSDK
+}
 
-			MLResult Result = MLGraphicsSetFrameTimingHint(GraphicsClient, FTH);
-			if (Result == MLResult_Ok)
+void FMagicLeapHMD::SetFrameTimingHint_RenderThread()
+{
+	check(IsInRenderingThread() || IsInRHIThread())
+#if WITH_MLSDK
+	FString EnumVal;
+	GConfig->GetString(TEXT("/Script/LuminRuntimeSettings.LuminRuntimeSettings"), TEXT("FrameTimingHint"), EnumVal, GEngineIni);
+
+	if (EnumVal.Len() > 0 && MLHandleIsValid(GraphicsClient))
+	{
+		const static UEnum* FrameTimingEnum = FindObject<UEnum>(ANY_PACKAGE, TEXT("ELuminFrameTimingHint"));
+		check(FrameTimingEnum);
+
+		const ELuminFrameTimingHint InFrameTimingHint = static_cast<ELuminFrameTimingHint>(FrameTimingEnum->GetValueByNameString(EnumVal));
+
+		if (InFrameTimingHint != CurrentFrameTimingHint)
+		{
+			if (GraphicsClient != ML_INVALID_HANDLE)
 			{
-				UE_LOG(LogMagicLeap, Log, TEXT("Set Lumin frame timing hint to %s."), *FrameTimingEnum->GetNameStringByValue((int64)InFrameTimingHint));
-				CurrentFrameTimingHint = InFrameTimingHint;
+				MLGraphicsFrameTimingHint FTH = MLGraphicsFrameTimingHint_Unspecified;
+				EngineGameThreadFPS = 0.0f;
+
+				switch (InFrameTimingHint)
+				{
+				case ELuminFrameTimingHint::Unspecified:
+					FTH = MLGraphicsFrameTimingHint_Unspecified;
+					EngineGameThreadFPS = 0.0f;
+					break;
+				case ELuminFrameTimingHint::Maximum:
+					FTH = MLGraphicsFrameTimingHint_Maximum;
+					EngineGameThreadFPS = 122.0f;
+					break;
+				case ELuminFrameTimingHint::FPS_60:
+					FTH = MLGraphicsFrameTimingHint_60Hz;
+					EngineGameThreadFPS = 62.0f;
+					break;
+				case ELuminFrameTimingHint::FPS_120:
+					FTH = MLGraphicsFrameTimingHint_120Hz;
+					EngineGameThreadFPS = 122.0f;
+					break;
+				default:
+					FTH = MLGraphicsFrameTimingHint_Unspecified;
+					UE_LOG(LogMagicLeap, Warning, TEXT("Tried to set invalid Frame Timing Hint!  Defaulting to unspecified."));
+				}
+
+				MLResult Result = MLGraphicsSetFrameTimingHint(GraphicsClient, FTH);
+				if (Result == MLResult_Ok)
+				{
+					UE_LOG(LogMagicLeap, Display, TEXT("Set Lumin frame timing hint to %s."), *FrameTimingEnum->GetNameStringByValue((int64)InFrameTimingHint));
+					CurrentFrameTimingHint = InFrameTimingHint;
+				}
+				else
+				{
+					UE_LOG(LogMagicLeap, Log, TEXT("Failed to set Lumin frame timing hint to %s with error %s!"), *FrameTimingEnum->GetNameStringByValue((int64)InFrameTimingHint), UTF8_TO_TCHAR(MLGetResultString(Result)));
+				}
 			}
 			else
 			{
-				UE_LOG(LogMagicLeap, Log, TEXT("Failed to set Lumin frame timing hint to %s...invalid parameter!"), *FrameTimingEnum->GetNameStringByValue((int64)InFrameTimingHint));
+				UE_LOG(LogMagicLeap, Warning, TEXT("Failed to set Lumin frame timing hint.  Invalid context."));
 			}
 		}
-		else
-		{
-			UE_LOG(LogMagicLeap, Warning, TEXT("Failed to set Lumin frame timing hint.  Invalid context."));
-		}
 	}
+#endif //WITH_MLSDK
 }
-#endif //PLATFORM_LUMIN
 
 float FMagicLeapHMD::GetWorldToMetersScale() const
 {
@@ -927,8 +1115,7 @@ FMatrix FMagicLeapHMD::GetStereoProjectionMatrix(const enum EStereoscopicPass St
 	check(IsStereoEnabled());
 	const int viewport = (StereoPassType == eSSP_LEFT_EYE) ? 0 : 1;
 	const FTrackingFrame& frame = GetCurrentFrame();
-	// TODO: Remove this for vulkan when we can get a better result from the frame
-	return (bDeviceInitialized && !IsVulkanPlatform(GMaxRHIShaderPlatform)) ? MagicLeap::ToFMatrix(frame.UpdateInfoArray.virtual_camera_extents[viewport].projection) : FMatrix::Identity;
+	return bDeviceInitialized ? MagicLeap::ToUEProjectionMatrix(frame.UpdateInfoArray.virtual_camera_extents[viewport].projection, GNearClippingPlane) : FMatrix::Identity;
 #else
 	return FMatrix();
 #endif //WITH_MLSDK
@@ -949,7 +1136,7 @@ void FMagicLeapHMD::UpdateViewportRHIBridge(bool /* bUseSeparateRenderTarget */,
 	}
 }
 
-bool FMagicLeapHMD::GetHeadTrackingState(FHeadTrackingState& State) const
+bool FMagicLeapHMD::GetHeadTrackingState(FMagicLeapHeadTrackingState& State) const
 {
 	if (bHeadTrackingStateAvailable)
 	{
@@ -958,6 +1145,15 @@ bool FMagicLeapHMD::GetHeadTrackingState(FHeadTrackingState& State) const
 	return bHeadTrackingStateAvailable;
 }
 
+bool FMagicLeapHMD::GetHeadTrackingMapEvents(TSet<EMagicLeapHeadTrackingMapEvent>& MapEvents) const
+{
+	if (bHeadposeMapEventsAvailable)
+	{
+		MapEvents = HeadposeMapEvents;
+	}
+	return bHeadposeMapEventsAvailable;
+}
+
 void FMagicLeapHMD::UpdateNearClippingPlane()
 {
 	GNearClippingPlane = (GameTrackingFrame.NearClippingPlane > GNearClippingPlane) ? GameTrackingFrame.NearClippingPlane : GNearClippingPlane;
@@ -970,20 +1166,6 @@ FMagicLeapCustomPresent* FMagicLeapHMD::GetActiveCustomPresent(const bool bRequi
 		return nullptr;
 	}
 
-#if PLATFORM_WINDOWS
-	if (CustomPresentD3D11)
-	{
-		return CustomPresentD3D11;
-	}
-#endif // PLATFORM_WINDOWS
-
-#if PLATFORM_MAC
-	if (CustomPresentMetal)
-	{
-		return CustomPresentMetal;
-}
-#endif // PLATFORM_MAC
-
 #if PLATFORM_WINDOWS || PLATFORM_LINUX || PLATFORM_LUMIN
 	if (CustomPresentOpenGL)
 	{
@@ -998,6 +1180,13 @@ FMagicLeapCustomPresent* FMagicLeapHMD::GetActiveCustomPresent(const bool bRequi
 	}
 #endif //PLATFORM_WINDOWS ||  PLATFORM_LUMIN
 
+#if PLATFORM_MAC
+	if (CustomPresentMetal)
+	{
+		return CustomPresentMetal;
+	}
+#endif // PLATFORM_MAC
+
 	return nullptr;
 }
 
@@ -1039,22 +1228,12 @@ bool FMagicLeapHMD::AllocateRenderTargetTexture(uint32 Index, uint32 SizeX, uint
 		return false;
 	}
 
-#if PLATFORM_MAC
-	// TODO: fix for Mac when VDZI is supported on Metal.
-	return false;
-#else
 	FRHIResourceCreateInfo CreateInfo;
 	RHICreateTargetableShaderResource2D(SizeX, SizeY, PF_R8G8B8A8, 1, TexCreate_None, TexCreate_RenderTargetable, false, CreateInfo, OutTargetableTexture, OutShaderResourceTexture);
 
 	return true;
-#endif
 }
 
-//FRHICustomPresent* FMagicLeapHMD::GetCustomPresent()
-//{
-//	return GetActiveCustomPresent();
-//}
-
 FMagicLeapHMD::FMagicLeapHMD(IMagicLeapPlugin* InMagicLeapPlugin, IARSystemSupport* InARImplementation, bool bEnableVDZI, bool bUseVulkan) :
 	// We don't do any mirroring on Lumin as we render direct to the device only.
 	FHeadMountedDisplayBase(InARImplementation),
@@ -1094,10 +1273,17 @@ FMagicLeapHMD::FMagicLeapHMD(IMagicLeapPlugin* InMagicLeapPlugin, IARSystemSuppo
 	bIsVDZIEnabled(bEnableVDZI),
 	bUseVulkanForZI(bUseVulkan),
 	bVDZIWarningDisplayed(false),
+#if PLATFORM_WINDOWS || PLATFORM_LINUX || PLATFORM_MAC
+	ZIServerPingState(static_cast<int32>(EServerPingState::NotInitialized)),
+#endif // PLATFORM_WINDOWS || PLATFORM_LINUX || PLATFORM_MAC
 	bPrivilegesEnabled(false),
 	CurrentFrameTimingHint(ELuminFrameTimingHint::Unspecified),
+	EngineGameThreadFPS(0.0f),
 	bQueuedGraphicsCreateCall(false),
-	bHeadTrackingStateAvailable(false)
+	bHeadTrackingStateAvailable(false),
+	bHeadposeMapEventsAvailable(false),
+	SavedMaxFPS(0.0f),
+	DefaultRenderTargetSize(1280 * 2, 960)
 {
 #if WITH_EDITOR
 	World = nullptr;
@@ -1112,8 +1298,6 @@ FMagicLeapHMD::~FMagicLeapHMD()
 
 void FMagicLeapHMD::Startup()
 {
-	LoadFromIni();
-
 	// grab a pointer to the renderer module for displaying our mirror window
 	const FName RendererModuleName("Renderer");
 	RendererModule = FModuleManager::GetModulePtr<IRendererModule>(RendererModuleName);
@@ -1135,23 +1319,6 @@ void FMagicLeapHMD::Startup()
 		PixelDensity = FMath::Clamp(PixelDensityCVar->GetFloat(), PixelDensityMin, PixelDensityMax);
 	}
 
-#if PLATFORM_WINDOWS
-	if (IsPCPlatform(GMaxRHIShaderPlatform) && !IsOpenGLPlatform(GMaxRHIShaderPlatform) && !IsVulkanPlatform(GMaxRHIShaderPlatform))
-	{
-		UE_LOG(LogMagicLeap, Display, TEXT("Creating FMagicLeapCustomPresentD3D11"));
-		CustomPresentD3D11 = new FMagicLeapCustomPresentD3D11(this);
-	}
-#endif // PLATFORM_WINDOWS
-
-#if PLATFORM_MAC
-	if (IsMetalPlatform(GMaxRHIShaderPlatform) && !IsOpenGLPlatform(GMaxRHIShaderPlatform))
-	{
-		UE_LOG(LogMagicLeap, Display, TEXT("Creating FMagicLeapCustomPresentMetal"));
-		//DISABLED until complete
-		//CustomPresentMetal = new FMagicLeapCustomPresentMetal(this);
-	}
-#endif // PLATFORM_MAC
-
 #if PLATFORM_WINDOWS || PLATFORM_LINUX || PLATFORM_LUMIN
 	if (IsOpenGLPlatform(GMaxRHIShaderPlatform))
 	{
@@ -1168,6 +1335,14 @@ void FMagicLeapHMD::Startup()
 	}
 #endif // PLATFORM_WINDOWS || PLATFORM_LUMIN
 
+#if PLATFORM_MAC
+	if (IsMetalPlatform(GMaxRHIShaderPlatform))
+	{
+		UE_LOG(LogMagicLeap, Display, TEXT("Creating FMagicLeapCustomPresentMetal"));
+		CustomPresentMetal = new FMagicLeapCustomPresentMetal(this);
+	}
+#endif // PLATFORM_MAC
+
 	UE_LOG(LogMagicLeap, Log, TEXT("MagicLeap initialized."));
 }
 
@@ -1203,24 +1378,6 @@ void FMagicLeapHMD::LoadFromIni()
 		WindowMirrorMode = WindowMirrorModeValue;
 	}
 #endif
-
-#if PLATFORM_LUMIN
-	ELuminFrameTimingHint ConfigFrameTimingHint = ELuminFrameTimingHint::Unspecified;
-	const static UEnum* FrameTimingEnum = StaticEnum<ELuminFrameTimingHint>();
-
-	FString EnumVal;
-	GConfig->GetString(TEXT("/Script/LuminRuntimeSettings.LuminRuntimeSettings"), TEXT("FrameTimingHint"), EnumVal, GEngineIni);
-
-	if (EnumVal.Len() > 0)
-	{
-		// This will be set later during Graphics Client initialization based on thje value in FrameTimingHint read from the config here.
-		ConfigFrameTimingHint = static_cast<ELuminFrameTimingHint>(FrameTimingEnum->GetValueByNameString(EnumVal));
-		if (GraphicsClient != ML_INVALID_HANDLE)
-		{
-			SetFrameTimingHint(ConfigFrameTimingHint);
-		}
-	}
-#endif //PLATFORM_LUMIN
 }
 
 void FMagicLeapHMD::SaveToIni()
@@ -1259,6 +1416,9 @@ void FMagicLeapHMD::OnBeginPlay(FWorldContext& InWorldContext)
 	SetIgnoreInput(true);
 #endif
 
+	// Reset CurrentFrameTimingHint so that proper value from config can be updated after device is initialized.
+	CurrentFrameTimingHint = ELuminFrameTimingHint::Unspecified;
+
 	EnableDeviceFeatures();
 }
 
@@ -1289,7 +1449,7 @@ void FMagicLeapHMD::EnableDeviceFeatures()
 		EnablePrivileges();
 		EnablePerception();
 		EnableHeadTracking();
-		EnableInputDevices();
+		EnableTrackerEntities();
 
 		// We also avoid enabling the custom profile when there's no HMD, as otherwise
 		// we get the profile effects on non-vr-preview rendering.
@@ -1301,9 +1461,8 @@ void FMagicLeapHMD::DisableDeviceFeatures()
 {
 	AppFramework.OnApplicationShutdown();
 	RestoreBaseProfile();
-	DisableInputDevices();
+	DisableTrackerEntities();
 	DisableHeadTracking();
-	DISABLE_MAGIC_LEAP_MODULE("MagicLeapEyeTracker");
 	DisablePerception();
 	DisablePrivileges();
 	if (GIsEditor)
@@ -1314,11 +1473,6 @@ void FMagicLeapHMD::DisableDeviceFeatures()
 	bVDZIWarningDisplayed = false;
 }
 
-#if (PLATFORM_WINDOWS && WITH_MLSDK)
-ML_EXTERN_C_BEGIN
-ML_API MLResult ML_CALL MLGraphicsCreateClientVk(const MLGraphicsOptions *options, void *vulkan_instance, void *vulkan_physical_device, void *vulkan_logical_device, MLHandle *out_graphics_client);
-ML_EXTERN_C_END
-#endif // (PLATFORM_WINDOWS && WITH_MLSDK)
 
 void FMagicLeapHMD::InitDevice_RenderThread()
 {
@@ -1351,8 +1505,43 @@ void FMagicLeapHMD::InitDevice_RenderThread()
 #if PLATFORM_MAC
 		if (IsMetalPlatform(GMaxRHIShaderPlatform) && !IsOpenGLPlatform(GMaxRHIShaderPlatform))
 		{
-			bDeviceSuccessfullyInitialized = true;
-			//UE_LOG(LogMagicLeap, Error, TEXT("FMagicLeapCustomPresentMetal is not supported."));
+			bQueuedGraphicsCreateCall = true;
+
+			ExecuteOnRHIThread_DoNotWait([this, gfx_opts]()
+			{
+				UE_LOG(LogMagicLeap, Display, TEXT("FMagicLeapCustomPresentMetal is supported."));
+
+				FMetalDynamicRHI* MetalDynamicRHI = static_cast<FMetalDynamicRHI*>(GDynamicRHI);
+                check(MetalDynamicRHI != nullptr);
+
+				id<MTLDevice> MetalDevice = (id<MTLDevice>)MetalDynamicRHI->RHIGetNativeDevice();
+
+				MLGraphicsOptions Options = gfx_opts;
+
+				MLRemoteGraphicsCreateClientMTLParams MTLParams;
+				MLRemoteGraphicsCreateClientMTLParamsInit(&MTLParams);
+				MTLParams.options = &Options;
+				MTLParams.metal_device = &MetalDevice;
+
+				GraphicsClient = ML_INVALID_HANDLE;
+				MLResult Result = MLResult_Ok;
+
+				Result = MLRemoteGraphicsCreateClientMTL(&MTLParams, &GraphicsClient);
+				if (Result == MLResult_Ok)
+				{
+					InitializeClipExtents_RenderThread();
+					SetFrameTimingHint_RenderThread();
+				}
+				else
+				{
+					GraphicsClient = ML_INVALID_HANDLE;
+					UE_LOG(LogMagicLeap, Error, TEXT("MLGraphicsCreateClientVk failed with status %s"), UTF8_TO_TCHAR(MLGetResultString(Result)));
+				}
+
+				FPlatformAtomics::InterlockedExchange(&bDeviceInitialized, Result == MLResult_Ok);
+				FPlatformAtomics::InterlockedExchange(&bDeviceWasJustInitialized, Result == MLResult_Ok);
+				bQueuedGraphicsCreateCall = false;
+			});
 		}
 #endif // PLATFORM_MAC
 
@@ -1368,6 +1557,7 @@ void FMagicLeapHMD::InitDevice_RenderThread()
 			{
 				bDeviceSuccessfullyInitialized = true;
 				InitializeClipExtents_RenderThread();
+				SetFrameTimingHint_RenderThread();
 			}
 			else
 			{
@@ -1394,14 +1584,12 @@ void FMagicLeapHMD::InitDevice_RenderThread()
 				uint64 LogicalDevice = VulkanRHIBridge::GetLogicalDevice(VulkanDevice);
 				GraphicsClient = ML_INVALID_HANDLE;
 				MLResult Result = MLResult_Ok;
-#if PLATFORM_LUMIN
+
 				Result = MLGraphicsCreateClientVk(&gfx_opts, (VkInstance)Instance, (VkPhysicalDevice)PhysicalDevice, (VkDevice)LogicalDevice, &GraphicsClient);
-#else
-				Result = MLGraphicsCreateClientVk(&gfx_opts, (void*)Instance, (void*)PhysicalDevice, (void*)LogicalDevice, &GraphicsClient);
-#endif
 				if (Result == MLResult_Ok)
 				{
 					InitializeClipExtents_RenderThread();
+					SetFrameTimingHint_RenderThread();
 				}
 				else
 				{
@@ -1420,14 +1608,6 @@ void FMagicLeapHMD::InitDevice_RenderThread()
 			FPlatformAtomics::InterlockedExchange(&bDeviceInitialized, bDeviceSuccessfullyInitialized);
 			FPlatformAtomics::InterlockedExchange(&bDeviceWasJustInitialized, bDeviceSuccessfullyInitialized);
 		}
-
-#if PLATFORM_LUMIN
-		// Initialize the frame timing hint, if we got a successful graphics client initialization
-		if (GraphicsClient != ML_INVALID_HANDLE)
-		{
-			SetFrameTimingHint(CurrentFrameTimingHint);
-		}
-#endif //PLATFORM_LUMIN
 #endif // PLATFORM_WINDOWS || PLATFORM_LINUX || PLATFORM_LUMIN
 	}
 #endif //WITH_MLSDK
@@ -1440,6 +1620,12 @@ void FMagicLeapHMD::InitDevice()
 		// If the HMD is not connected don't bother initializing the render device since the VDZI graphics calls freeze the editor if the VDZI server is not running.
 		if (IsHMDConnected())
 		{
+			// Enable HMD as it could have been disabled on a previous VRPreview failure
+			EnableHMD(true);
+
+			// Save the default MaxFPS value to be restored when we shutdown VRPreview.
+			SavedMaxFPS = GEngine->GetMaxFPS();
+
 			FMagicLeapHMD* This = this;
 			ENQUEUE_RENDER_COMMAND(InitDevice)(
 				[This](FRHICommandList& RHICmdList)
@@ -1465,9 +1651,17 @@ void FMagicLeapHMD::InitDevice()
 		{
 			// This init must happen on the main thread for VR preview, otherwise it crashes on a non-Lumin RHI.
 
-				// Save any runtime configuration changes from the .ini.
+			// Save any runtime configuration changes from the .ini.
 			LoadFromIni();
 
+			// Moved here from SetFrameTimingHint() because it crashes on non-game threads.
+			if (GEngine != nullptr)
+			{
+				// Don't rely on MLGraphicsBeginFrame to manage app's frame rate cadence because when the device goes into standby mode
+				// we skip calls to it which then causes the render thread to run at it's max pace.
+				GEngine->SetMaxFPS(EngineGameThreadFPS);
+			}
+
 			// VD/ZI works best in windowed mode since it can sometimes be used in conjunction with the mock ml1 device's window.
 #if PLATFORM_LUMIN
 			EWindowMode::Type WindowMode = EWindowMode::Fullscreen;
@@ -1496,9 +1690,18 @@ void FMagicLeapHMD::ReleaseDevice()
 {
 	check(IsInGameThread());
 
+	//Disable the HMD, reverting flags for potential followup in PIE 
+	EnableHMD(false);
+
 	// save any runtime configuration changes to the .ini
 	SaveToIni();
 
+	// Restore original max fps after shutting down VRPreview.
+	if (GEngine != nullptr)
+	{
+		GEngine->SetMaxFPS(SavedMaxFPS);
+	}
+
 	FMagicLeapHMD* Plugin = this;
 	ENQUEUE_RENDER_COMMAND(ReleaseDevice_RT)(
 		[Plugin](FRHICommandListImmediate& RHICmdList)
@@ -1521,11 +1724,7 @@ void FMagicLeapHMD::ReleaseDevice_RenderThread()
 	{
 		FPlatformAtomics::InterlockedExchange(&bDeviceInitialized, 0);
 
-#if PLATFORM_WINDOWS
-		if (CustomPresentD3D11)
-		{
-			CustomPresentD3D11->Reset();
-		}
+#if PLATFORM_WINDOWS || PLATFORM_LUMIN
 		if (CustomPresentOpenGL)
 		{
 			CustomPresentOpenGL->Reset();
@@ -1534,38 +1733,31 @@ void FMagicLeapHMD::ReleaseDevice_RenderThread()
 		{
 			CustomPresentVulkan->Reset();
 		}
-#elif PLATFORM_MAC
-		if (CustomPresentMetal)
-		{
-			CustomPresentMetal->Reset();
-		}
 #elif PLATFORM_LINUX
 		if (CustomPresentOpenGL)
 		{
 			CustomPresentOpenGL->Reset();
 		}
-#else
-		if (CustomPresentOpenGL)
-		{
-			CustomPresentOpenGL->Reset();
-		}
-		if (CustomPresentVulkan)
+#elif PLATFORM_MAC
+		if (CustomPresentMetal)
 		{
-			CustomPresentVulkan->Reset();
+			CustomPresentMetal->Reset();
 		}
 #endif
 
-#if WITH_MLSDK
-		MLResult Result = MLGraphicsDestroyClient(&GraphicsClient);
-		if (Result != MLResult_Ok)
+		ExecuteOnRHIThread([this]()
 		{
-			UE_LOG(LogMagicLeap, Error, TEXT("MLGraphicsDestroyClient failed with status %s"), UTF8_TO_TCHAR(MLGetResultString(Result)));
-		}
-		else
-		{
-			UE_LOG(LogMagicLeap, Display, TEXT("Graphics client destroyed successfully."));
-		}
+#if WITH_MLSDK
+			if (MLHandleIsValid(GraphicsClient))
+			{
+				MLResult Result = MLGraphicsDestroyClient(&GraphicsClient);
+				UE_CLOG(Result != MLResult_Ok, LogMagicLeap, Error, TEXT("MLGraphicsDestroyClient failed with status %s"), UTF8_TO_TCHAR(MLGetResultString(Result)));
+				UE_CLOG(Result == MLResult_Ok, LogMagicLeap, Display, TEXT("Graphics client destroyed successfully."));
+			}
+
+			GraphicsClient = ML_INVALID_HANDLE;
 #endif //WITH_MLSDK
+		});
 	}
 }
 
@@ -1588,7 +1780,7 @@ bool FMagicLeapHMD::GetRelativeEyePose(int32 DeviceId, EStereoscopicPass Eye, FQ
 		const FTrackingFrame& Frame = GetCurrentFrame();
 		const int EyeIdx = (Eye == eSSP_LEFT_EYE) ? 0 : 1;
 
-		const FTransform EyeToWorld = MagicLeap::ToFTransform(Frame.RenderInfoArray.virtual_cameras[EyeIdx].transform, Frame.WorldToMetersScale);		// "world" here means the HMDs tracking space
+		const FTransform EyeToWorld = MagicLeap::ToFTransform(Frame.FrameInfo.virtual_camera_info_array.virtual_cameras[EyeIdx].transform, Frame.WorldToMetersScale);		// "world" here means the HMDs tracking space
 		const FTransform EyeToHMD = EyeToWorld * Frame.RawPose.Inverse();		// RawPose is HMDToWorld
 		OutPosition = EyeToHMD.GetTranslation();
 		OutOrientation = EyeToHMD.GetRotation();
@@ -1644,17 +1836,15 @@ void FMagicLeapHMD::OnBeginRendering_GameThread()
 		MLSnapshot* OldSnapshot = RenderTrackingFrame.Snapshot;
 		// Don't update RenderTrackingFrame here. It is updated by the RHITrackingFrame in FMagicLeapCustomPreset::BeginRendering()
 		// RenderTrackingFrame = TrackingFrameCopy;
-#if !PLATFORM_MAC
 		ExecuteOnRHIThread_DoNotWait([this, TrackingFrameCopy]()
 		{
 			RHITrackingFrame = TrackingFrameCopy;
 		});
-#endif //PLATFORM_MAC
 	});
 #endif //WITH_MLSDK
 
-	// Update the devices, in particular input controller devices.
-	IMagicLeapPlugin::Get().OnBeginRendering_GameThread_UpdateInputDevices();
+	// Update the trackers
+	static_cast<FMagicLeapPlugin*>(&IMagicLeapPlugin::Get())->OnBeginRendering_GameThread();
 }
 
 TSharedPtr<class IXRCamera, ESPMode::ThreadSafe> FMagicLeapHMD::GetXRCamera(int32 DeviceId /*= HMDDeviceId*/)
@@ -1671,6 +1861,8 @@ void FMagicLeapHMD::OnBeginRendering_RenderThread(FRHICommandListImmediate& RHIC
 {
 	check(IsInRenderingThread());
 
+	GetClipExtents();
+
 	if (GetActiveCustomPresent())
 	{
 		GetActiveCustomPresent()->BeginRendering();
@@ -1683,6 +1875,14 @@ void FMagicLeapHMD::RenderTexture_RenderThread(class FRHICommandListImmediate& R
 #if WITH_MLSDK
 	check(IsInRenderingThread());
 
+	if (GetActiveCustomPresent() != nullptr)
+	{
+		// If custom present wants to render using Unreal's pipeline instead of direct blit or texture copy functions
+		GetActiveCustomPresent()->RenderToMLSurfaces_RenderThread(RHICmdList, SrcTexture);
+	}
+
+	// TODO: Refactor the code below. FMagicLeapCustomPresent::RenderToTextureSlice_RenderThread() is very similar.
+
 	// If we aren't mirroring there's nothing to do as the actual render on device
 	// happens in the custom presenter.
 	if (WindowMirrorMode > 0)
@@ -1778,30 +1978,9 @@ void FMagicLeapHMD::SetClippingPlanes(float NCP, float FCP)
 	UpdateNearClippingPlane();
 }
 
-bool FMagicLeapHMD::IsInitialized() const
-{
-	return (AppFramework.IsInitialized());
-}
-
 void FMagicLeapHMD::ShutdownRendering()
 {
 	check(IsInRenderingThread());
-#if PLATFORM_WINDOWS
-	if (CustomPresentD3D11.GetReference())
-	{
-		CustomPresentD3D11->Reset();
-		CustomPresentD3D11->Shutdown();
-		CustomPresentD3D11 = nullptr;
-	}
-#endif // PLATFORM_WINDOWS
-#if PLATFORM_MAC
-	if (CustomPresentMetal.GetReference())
-	{
-		CustomPresentMetal->Reset();
-		CustomPresentMetal->Shutdown();
-		CustomPresentMetal = nullptr;
-}
-#endif // PLATFORM_MAC
 #if PLATFORM_WINDOWS || PLATFORM_LINUX || PLATFORM_LUMIN
 	if (CustomPresentOpenGL.GetReference())
 	{
@@ -1818,6 +1997,14 @@ void FMagicLeapHMD::ShutdownRendering()
 		CustomPresentVulkan = nullptr;
 	}
 #endif // PLATFORM_WINDOWS || PLATFORM_LUMIN
+#if PLATFORM_MAC
+	if (CustomPresentMetal.GetReference())
+	{
+		CustomPresentMetal->Reset();
+		CustomPresentMetal->Shutdown();
+		CustomPresentMetal = nullptr;
+	}
+#endif // PLATFORM_MAC
 }
 
 FTrackingFrame& FMagicLeapHMD::GetCurrentFrameMutable()
@@ -1885,9 +2072,22 @@ FAppFramework& FMagicLeapHMD::GetAppFramework()
 	return AppFramework;
 }
 
-void FMagicLeapHMD::SetFocusActor(const AActor* InFocusActor)
+void FMagicLeapHMD::SetFocusActor(const AActor* InFocusActor, bool bSetStabilizationActor)
 {
 	FocusActor = InFocusActor;
+	if (bSetStabilizationActor)
+	{
+		StabilizationDepthActor = InFocusActor;
+	}
+}
+
+void FMagicLeapHMD::SetStabilizationDepthActor(const AActor* InStabilizationDepthActor, bool bSetFocusActor)
+{
+	StabilizationDepthActor = InStabilizationDepthActor;
+	if (bSetFocusActor)
+	{
+		FocusActor = InStabilizationDepthActor;
+	}
 }
 
 bool FMagicLeapHMD::IsPerceptionEnabled() const
@@ -2000,7 +2200,6 @@ void FMagicLeapHMD::DisablePrivileges()
 #if WITH_MLSDK
 	if (bPrivilegesEnabled)
 	{
-		UE_LOG(LogMagicLeap, Warning, TEXT("FMagicLeapHMD::DisablePrivileges"));
 		MLResult Result = MLPrivilegesShutdown();
 		UE_CLOG(Result != MLResult_Ok, LogMagicLeap, Error, TEXT("MLPrivilegesShutdown() "
 			"failed with error %s"), UTF8_TO_TCHAR(MLPrivilegesGetResultString(Result)));
@@ -2008,14 +2207,20 @@ void FMagicLeapHMD::DisablePrivileges()
 #endif // WITH_MLSDK
 }
 
-void FMagicLeapHMD::EnableInputDevices()
+void FMagicLeapHMD::EnableTrackerEntities()
 {
-	IMagicLeapPlugin::Get().EnableInputDevices();
+#if WITH_MLSDK
+	if (GameTrackingFrame.Snapshot == nullptr)
+	{
+		RefreshTrackingFrame();
+	}
+#endif // WITH_MLSDK
+	static_cast<FMagicLeapPlugin*>(&IMagicLeapPlugin::Get())->EnableTrackerEntities();
 }
 
-void FMagicLeapHMD::DisableInputDevices()
+void FMagicLeapHMD::DisableTrackerEntities()
 {
-	IMagicLeapPlugin::Get().DisableInputDevices();
+	static_cast<FMagicLeapPlugin*>(&IMagicLeapPlugin::Get())->DisableTrackerEntities();
 }
 
 void FMagicLeapHMD::EnablePerception()
@@ -2093,6 +2298,10 @@ void FMagicLeapHMD::DisableHeadTracking()
 		MLResult Result = MLHeadTrackingDestroy(HeadTracker);
 		UE_CLOG(Result != MLResult_Ok, LogMagicLeap, Error, TEXT("MLHeadTrackingDestroy failed with error %s."), UTF8_TO_TCHAR(MLGetResultString(Result)));
 		HeadTracker = ML_INVALID_HANDLE;
+
+		// Reset HeadTracking flags in case game calls GetHeadTrackingState()/GetHeadTrackingMapEvents() after head tracking is disabled. Also helps to reset the values for next VRPreview run.
+		bHeadTrackingStateAvailable = false;
+		bHeadposeMapEventsAvailable = false;
 	}
 #endif //WITH_MLSDK
 }
@@ -2107,59 +2316,20 @@ void FMagicLeapHMD::InitializeClipExtents_RenderThread()
 		GameTrackingFrame.NearClippingPlane = RenderTargetInfo.min_clip * GameTrackingFrame.WorldToMetersScale;
 		GameTrackingFrame.RecommendedFarClippingPlane = RenderTargetInfo.max_clip * GameTrackingFrame.WorldToMetersScale;
 		UpdateNearClippingPlane();
+
+		DefaultRenderTargetSize.X = 0;
+		for (uint32 i = 0; i < 2 && i < RenderTargetInfo.num_virtual_cameras; ++i)
+		{
+			DefaultRenderTargetSize.X += RenderTargetInfo.buffers[i].color.width;
+			DefaultRenderTargetSize.Y = RenderTargetInfo.buffers[i].color.height;
+		}
 	}
 	else
 	{
 		UE_LOG(LogMagicLeap, Error, TEXT("MLGraphicsGetRenderTargets failed with error %s"), UTF8_TO_TCHAR(MLGetResultString(Result)));
 	}
 
-	// get the clip extents for clipping content in update thread
-	Result = MLGraphicsGetClipExtents(GraphicsClient, &GameTrackingFrame.UpdateInfoArray);
-	if (Result != MLResult_Ok)
-	{
-		FString ErrorMesg = FString::Printf(TEXT("MLGraphicsGetClipExtents failed with error %s"), UTF8_TO_TCHAR(MLGetResultString(Result)));
-
-		// In case we're running under VD/ZI, there's always the risk of disconnects.
-		// In those cases, the graphics API can return an error, but the client handle might still be valid.
-		// So we need to ensure that we always have valid data to prevent any Nan-related errors.
-		// On Lumin, we'll just assert.
-#if PLATFORM_WINDOWS || PLATFORM_LINUX || PLATFORM_MAC
-		GameTrackingFrame.Handle = ML_INVALID_HANDLE;
-		MagicLeap::ResetClipExtentsInfoArray(GameTrackingFrame.UpdateInfoArray);
-		UE_LOG(LogMagicLeap, Error, TEXT("%s"), *ErrorMesg);
-#else
-		UE_LOG(LogMagicLeap, Fatal, TEXT("%s"), *ErrorMesg);
-#endif
-	}
-
-	/* Expected Right Handed Projection Model */
-	/*
-	MLGraphicsProjectionType_ReversedInfiniteZ
-	proj_mat[2][2] = 0.0f;
-	proj_mat[2][3] = -1.0f;
-	proj_mat[3][2] = near_clip_meters;
-	*/
-
-	/* Convert full extents from Graphics Projection Model to Unreal Projection Model */
-	// Graphics returns values in Infinite Z. We convert it to Reversed Infinite Z here.
-	GameTrackingFrame.UpdateInfoArray.full_extents.projection.matrix_colmajor[10] = 0.0f; // Model change hack
-	GameTrackingFrame.UpdateInfoArray.full_extents.projection.matrix_colmajor[11] = -1.0f; // Model change hack
-
-	// We also convert the near plane into centimeters since Unreal directly uses these values
-	// for various calculations such as shadow algorithm and expects units to be in centimeters
-	GameTrackingFrame.UpdateInfoArray.full_extents.projection.matrix_colmajor[14] = GNearClippingPlane; // Model change hack
-
-	/* Convert eye extents from Graphics Projection Model to Unreal Projection Model */
-	for (uint32_t eye = 0; eye < GameTrackingFrame.UpdateInfoArray.num_virtual_cameras; ++eye)
-	{
-		// Graphics returns values in Infinite Z. We convert it to Reversed Infinite Z here.
-		GameTrackingFrame.UpdateInfoArray.virtual_camera_extents[eye].projection.matrix_colmajor[10] = 0.0f; // Model change hack
-		GameTrackingFrame.UpdateInfoArray.virtual_camera_extents[eye].projection.matrix_colmajor[11] = -1.0f; // Model change hack
-
-		// We also convert the near plane into centimeters since Unreal directly uses these values
-		// for various calculations such as shadow algorithm and expects units to be in centimeters
-		GameTrackingFrame.UpdateInfoArray.virtual_camera_extents[eye].projection.matrix_colmajor[14] = GNearClippingPlane; // Model change hack
-	}
+	GetClipExtents();
 
 	// TODO Apply snapshot head pose to all the update transforms because graphics does not apply pose
 	// But we currently use the last frame render transforms so this does not need to be done just yet
diff --git a/Engine/Plugins/Lumin/MagicLeap/Source/MagicLeap/Private/MagicLeapHMD.h b/Engine/Plugins/Lumin/MagicLeap/Source/MagicLeap/Private/MagicLeapHMD.h
--- a/Engine/Plugins/Lumin/MagicLeap/Source/MagicLeap/Private/MagicLeapHMD.h
+++ b/Engine/Plugins/Lumin/MagicLeap/Source/MagicLeap/Private/MagicLeapHMD.h
@@ -5,7 +5,6 @@
 #include "HeadMountedDisplayBase.h"
 #include "Runtime/Launch/Resources/Version.h"
 #include "IMagicLeapPlugin.h"
-#include "IMagicLeapHMD.h"
 #include "Misc/ScopeLock.h"
 #include "HAL/ThreadSafeBool.h"
 
@@ -14,37 +13,19 @@
 #include "AppFramework.h"
 #include "SceneViewExtension.h"
 #include "MagicLeapMath.h"
-#include "MagicLeapCustomPresent.h"
+#include "MagicLeapCustomPresentGL.h"
+#include "MagicLeapCustomPresentVk.h"
+#include "MagicLeapCustomPresentMetal.h"
 #include "MagicLeapHMDFunctionLibrary.h"
 #include "LuminRuntimeSettings.h"
-#include "MagicLeapPluginUtil.h" // for ML_INCLUDES_START/END
-
-#if WITH_MLSDK
-ML_INCLUDES_START
-#include <ml_head_tracking.h>
-ML_INCLUDES_END
-#endif //WITH_MLSDK
+#include "Lumin/CAPIShims/LuminAPIHeadTracking.h"
 
 class FXRTrackingSystemBase;
 
-/**
- * The public interface to this module.
- */
-class ILuminARModule : public IModuleInterface
-{
-public:
-	//create for mutual connection (regardless of construction order)
-	virtual TSharedPtr<IARSystemSupport, ESPMode::ThreadSafe> CreateARImplementation() = 0;
-	//Now connect (regardless of connection order)
-	virtual void ConnectARImplementationToXRSystem(FXRTrackingSystemBase* InXRTrackingSystem) = 0;
-	//Now initialize fully connected systems
-	virtual void InitializeARImplementation() = 0;
-};
-
 /**
   * MagicLeap Head Mounted Display
   */
-class MAGICLEAP_API FMagicLeapHMD : public IMagicLeapHMD, public FHeadMountedDisplayBase, public FXRRenderTargetManager
+class MAGICLEAP_API FMagicLeapHMD : public FHeadMountedDisplayBase, public FXRRenderTargetManager
 {
 public:
 
@@ -70,20 +51,10 @@ public:
 
 	virtual void ResetOrientationAndPosition(float yaw = 0.f) override;
 	virtual void ResetOrientation(float Yaw = 0.f) override;
-	virtual void ResetPosition() override;
 
 	virtual void OnBeginPlay(FWorldContext& InWorldContext) override;
 	virtual void OnEndPlay(FWorldContext& InWorldContext) override;
 
-	void SetBasePosition(const FVector& InBasePosition);
-	FVector GetBasePosition() const;
-
-	virtual void SetBaseRotation(const FRotator& BaseRot) override;
-	virtual FRotator GetBaseRotation() const override;
-
-	virtual void SetBaseOrientation(const FQuat& BaseOrient) override;
-	virtual FQuat GetBaseOrientation() const override;
-
 	float GetWorldToMetersScale() const override;
 
 	/** IHeadMountedDisplay interface */
@@ -148,10 +119,10 @@ public:
 	virtual ~FMagicLeapHMD();
 
 	/** FMagicLeapHMDBase interface */
-	virtual bool IsInitialized() const override;
 	bool IsDeviceInitialized() const { return (bDeviceInitialized != 0) ? true : false; }
 
-	bool GetHeadTrackingState(FHeadTrackingState& State) const;
+	bool GetHeadTrackingState(FMagicLeapHeadTrackingState& State) const;
+	bool GetHeadTrackingMapEvents(TSet<EMagicLeapHeadTrackingMapEvent>& MapEvents) const;
 	void UpdateNearClippingPlane();
 	FMagicLeapCustomPresent* GetActiveCustomPresent(const bool bRequireDeviceIsInitialized = true) const;
 	void ShutdownRendering();
@@ -164,32 +135,10 @@ public:
 
 	void PauseRendering(const bool bIsPaused) { bIsRenderingPaused = bIsPaused; }
 
-	/** Utility class to scope guard enabling and disabling game viewport client input processing.
-	* On creation it will enable the input processing, and on exit it will restore it to its
-	* previous state. Usage is:
-	*
-	* { FMagicLeapHMD::EnableInput EnableInput; PostSomeInputToMessageHandlers(); }
-	*/
-	struct EnableInput
-	{
-#if WITH_EDITOR
-		inline EnableInput()
-		{
-			SavedIgnoreInput = FMagicLeapHMD::GetHMD()->SetIgnoreInput(false);
-		}
-		inline ~EnableInput()
-		{
-			FMagicLeapHMD::GetHMD()->SetIgnoreInput(SavedIgnoreInput);
-		}
-
-	private:
-		bool SavedIgnoreInput;
-#endif
-	};
-
 public:
+	IRendererModule* GetRendererModule() { return RendererModule; }
 
-	uint32 GetViewportCount() const { return AppFramework.IsInitialized() ? AppFramework.GetViewportCount() : 0; }
+	uint32 GetViewportCount() const { return AppFramework.GetViewportCount(); }
 
 	// TODO: add const versions
 	FTrackingFrame& GetCurrentFrameMutable();
@@ -207,7 +156,21 @@ public:
 
 	const FAppFramework& GetAppFrameworkConst() const;
 	FAppFramework& GetAppFramework();
-	void SetFocusActor(const AActor* InFocusActor);
+   /**
+	* Set the actor whose location is used as the focus point. The focus distance is the distance from the HMD to the focus point.
+	* 
+	* @param InFocusActor			The actor that will be set as the new focus actor.
+	* @param bSetStabilizationActor  True if the function should set the stabilization depth actor to match the passed in focus actor. (RECOMMENDED TO STAY CHECKED)
+	*/
+	void SetFocusActor(const AActor* InFocusActor, bool bSetStabilizationActor = true);
+	static const float DefaultFocusDistance;
+   /**
+	* Set the actor whose location is used as the depth for timewarp stabilization.
+	* 
+	* @param InStabilizationDepthActor  The actor that will be set as the new stabilization depth actor.
+	* @param bSetFocusActor				True if the function should set the focus actor to match the passed in stabilization depth actor. (RECOMMENDED TO STAY CHECKED)
+	*/
+	void SetStabilizationDepthActor(const AActor* InStabilizationDepthActor, bool bSetFocusActor = true);
 
 	bool IsPerceptionEnabled() const;
 
@@ -249,8 +212,8 @@ private:
 	void EnablePrivileges();
 	void DisablePrivileges();
 
-	void EnableInputDevices();
-	void DisableInputDevices();
+	void EnableTrackerEntities();
+	void DisableTrackerEntities();
 
 	void EnablePerception();
 	void DisablePerception();
@@ -264,18 +227,18 @@ private:
 	void RefreshTrackingFrame();
 
 #if WITH_MLSDK
-	EHeadTrackingError MLToUnrealHeadTrackingError(MLHeadTrackingError error) const;
-	EHeadTrackingMode MLToUnrealHeadTrackingMode(MLHeadTrackingMode mode) const;
+	EMagicLeapHeadTrackingError MLToUnrealHeadTrackingError(MLHeadTrackingError error) const;
+	EMagicLeapHeadTrackingMode MLToUnrealHeadTrackingMode(MLHeadTrackingMode mode) const;
+	EMagicLeapHeadTrackingMapEvent MLToUnrealHeadTrackingMapEvent(MLHeadTrackingMapEvent mapevent) const;
 #endif //WITH_MLSDK
 
 #if !PLATFORM_LUMIN
 	void DisplayWarningIfVDZINotEnabled();
 #endif
 
-#if PLATFORM_LUMIN
-	/** Sets a frame timing hint, which tells the device what your target framerate is, to aid in prediction and jitter */
-	void SetFrameTimingHint(ELuminFrameTimingHint InFrameTimingHint);
-#endif //PLATFORM_LUMIN
+	void GetClipExtents();
+	// Gets the config value, converts to enum and then makes the c-api function call.
+	void SetFrameTimingHint_RenderThread();
 
 private:
 	int32 bDeviceInitialized; // RW on render thread, R on game thread
@@ -303,23 +266,32 @@ private:
 	bool bIsVDZIEnabled;
 	bool bUseVulkanForZI;
 	bool bVDZIWarningDisplayed;
+#if PLATFORM_WINDOWS || PLATFORM_LINUX || PLATFORM_MAC
+	enum class EServerPingState : int32
+	{
+		NotInitialized,
+		AttemptingFirstPing,
+		ServerOnline,
+		ServerOffline,
+	};
+	volatile int32 ZIServerPingState;
+	static constexpr float ServerPingInterval = 5.0f;
+#endif // PLATFORM_WINDOWS || PLATFORM_LINUX || PLATFORM_MAC
 	bool bPrivilegesEnabled;
 
 	/** Current hint to the Lumin system about what our target framerate should be */
 	ELuminFrameTimingHint CurrentFrameTimingHint;
+	float EngineGameThreadFPS;
 
-#if PLATFORM_WINDOWS
-	TRefCountPtr<FMagicLeapCustomPresentD3D11> CustomPresentD3D11;
-#endif // PLATFORM_WINDOWS
-#if PLATFORM_MAC
-	TRefCountPtr<FMagicLeapCustomPresentMetal> CustomPresentMetal;
-#endif // PLATFORM_MAC
 #if PLATFORM_WINDOWS || PLATFORM_LINUX || PLATFORM_LUMIN
 	TRefCountPtr<FMagicLeapCustomPresentOpenGL> CustomPresentOpenGL;
 #endif // PLATFORM_WINDOWS || PLATFORM_LINUX || PLATFORM_LUMIN
 #if PLATFORM_WINDOWS || PLATFORM_LUMIN
 	TRefCountPtr<FMagicLeapCustomPresentVulkan> CustomPresentVulkan;
 #endif // PLATFORM_LUMIN
+#if PLATFORM_MAC
+	TRefCountPtr<FMagicLeapCustomPresentMetal> CustomPresentMetal;
+#endif // PLATFORM_MAC
 
 public:
 	FTrackingFrame GameTrackingFrame;
@@ -329,6 +301,7 @@ public:
 
 private:
 	TWeakObjectPtr<const AActor> FocusActor;
+	TWeakObjectPtr<const AActor> StabilizationDepthActor;
 	FThreadSafeBool bQueuedGraphicsCreateCall;
 
 	struct SavedProfileState
@@ -345,9 +318,12 @@ private:
 	};
 	SavedProfileState BaseProfileState;
 
-	FHeadTrackingState HeadTrackingState;
+	FMagicLeapHeadTrackingState HeadTrackingState;
 	bool bHeadTrackingStateAvailable;
 
+	bool bHeadposeMapEventsAvailable;
+	TSet<EMagicLeapHeadTrackingMapEvent> HeadposeMapEvents;
+
 #if WITH_EDITOR
 	/** The world we are playing. This is valid only within the span of BeginPlay & EndPlay. */
 	UWorld* World;
@@ -363,6 +339,9 @@ private:
 	/** Utility to get the MagicLeap specific HMD plugin instance. */
 	static FMagicLeapHMD* GetHMD();
 #endif
+
+	float SavedMaxFPS;
+	FIntPoint DefaultRenderTargetSize;
 };
 
 //DEFINE_LOG_CATEGORY_STATIC(LogHMD, Log, All);
diff --git a/Engine/Plugins/Lumin/MagicLeap/Source/MagicLeap/Private/MagicLeapHMDFunctionLibrary.cpp b/Engine/Plugins/Lumin/MagicLeap/Source/MagicLeap/Private/MagicLeapHMDFunctionLibrary.cpp
--- a/Engine/Plugins/Lumin/MagicLeap/Source/MagicLeap/Private/MagicLeapHMDFunctionLibrary.cpp
+++ b/Engine/Plugins/Lumin/MagicLeap/Source/MagicLeap/Private/MagicLeapHMDFunctionLibrary.cpp
@@ -3,9 +3,7 @@
 #include "MagicLeapHMDFunctionLibrary.h"
 #include "MagicLeapHMD.h"
 #include "Engine/Engine.h"
-#if WITH_MLSDK
-#include "ml_version.h"
-#endif //WITH_MLSDK
+#include "Lumin/CAPIShims/LuminAPILifecycle.h"
 
 static const FName MLDeviceName(TEXT("MagicLeap"));
 
@@ -27,38 +25,29 @@ static FMagicLeapHMD* GetMagicLeapHMD()
 }
 
 void UMagicLeapHMDFunctionLibrary::SetBasePosition(const FVector& InBasePosition)
-{
-	FMagicLeapHMD* const HMD = GetMagicLeapHMD();
-	if (HMD)
-	{
-		HMD->SetBasePosition(InBasePosition);
-	}
-}
+{ /* deprecated */ }
 
 void UMagicLeapHMDFunctionLibrary::SetBaseOrientation(const FQuat& InBaseOrientation)
-{
-	FMagicLeapHMD* const HMD = GetMagicLeapHMD();
-	if (HMD)
-	{
-		HMD->SetBaseOrientation(InBaseOrientation);
-	}
-}
+{ /* deprecated */ }
 
 void UMagicLeapHMDFunctionLibrary::SetBaseRotation(const FRotator& InBaseRotation)
+{ /* deprecated */ }
+
+void UMagicLeapHMDFunctionLibrary::SetFocusActor(const AActor* InFocusActor, bool bSetStabilizationActor)
 {
 	FMagicLeapHMD* const HMD = GetMagicLeapHMD();
 	if (HMD)
 	{
-		HMD->SetBaseRotation(InBaseRotation);
+		HMD->SetFocusActor(InFocusActor, bSetStabilizationActor);
 	}
 }
 
-void UMagicLeapHMDFunctionLibrary::SetFocusActor(const AActor* InFocusActor)
+void UMagicLeapHMDFunctionLibrary::SetStabilizationDepthActor(const AActor* InStabilizationDepthActor, bool bSetFocusActor)
 {
 	FMagicLeapHMD* const HMD = GetMagicLeapHMD();
 	if (HMD)
 	{
-		HMD->SetFocusActor(InFocusActor);
+		HMD->SetStabilizationDepthActor(InStabilizationDepthActor, bSetFocusActor);
 	}
 }
 
@@ -103,11 +92,19 @@ bool UMagicLeapHMDFunctionLibrary::IsRunningOnMagicLeapHMD()
 #if PLATFORM_LUMIN
 	return true;
 #else
+	FMagicLeapHMD* const HMD = GetMagicLeapHMD();
+	if (HMD)
+	{
+		//`IsStereoEnabled` is a strict form of `IsHMDEnabled`. 
+		//	Non-target platforms do not request stereo. 
+		return HMD->IsStereoEnabled();
+	}
+
 	return false;
 #endif
 }
 
-bool UMagicLeapHMDFunctionLibrary::GetHeadTrackingState(FHeadTrackingState& State)
+bool UMagicLeapHMDFunctionLibrary::GetHeadTrackingState(FMagicLeapHeadTrackingState& State)
 {
 	FMagicLeapHMD* const HMD = GetMagicLeapHMD();
 	if (HMD)
@@ -117,3 +114,23 @@ bool UMagicLeapHMDFunctionLibrary::GetHeadTrackingState(FHeadTrackingState& Stat
 
 	return false;
 }
+
+bool UMagicLeapHMDFunctionLibrary::GetHeadTrackingMapEvents(TSet<EMagicLeapHeadTrackingMapEvent>& MapEvents)
+{
+	FMagicLeapHMD* const HMD = GetMagicLeapHMD();
+	if (HMD)
+	{
+		return HMD->GetHeadTrackingMapEvents(MapEvents);
+	}
+
+	return false;
+}
+
+bool UMagicLeapHMDFunctionLibrary::SetAppReady()
+{
+#if WITH_MLSDK
+	return MLLifecycleSetReadyIndication() == MLResult_Ok;
+#else
+	return true;
+#endif // WITH_MLSDK
+}
diff --git a/Engine/Plugins/Lumin/MagicLeap/Source/MagicLeap/Private/MagicLeapHeadTrackingNotificationsComponent.cpp b/Engine/Plugins/Lumin/MagicLeap/Source/MagicLeap/Private/MagicLeapHeadTrackingNotificationsComponent.cpp
new file mode 100644
--- /dev/null
+++ b/Engine/Plugins/Lumin/MagicLeap/Source/MagicLeap/Private/MagicLeapHeadTrackingNotificationsComponent.cpp
@@ -0,0 +1,94 @@
+// Copyright 1998-2019 Epic Games, Inc. All Rights Reserved.
+
+#include "MagicLeapHeadTrackingNotificationsComponent.h"
+
+UMagicLeapHeadTrackingNotificationsComponent::UMagicLeapHeadTrackingNotificationsComponent()
+{
+	PrimaryComponentTick.bCanEverTick = true;
+	PrimaryComponentTick.bStartWithTickEnabled = true;
+	bAutoActivate = true;
+}
+
+void UMagicLeapHeadTrackingNotificationsComponent::TickComponent(float DeltaTime, enum ELevelTick TickType, FActorComponentTickFunction *ThisTickFunction)
+{
+	Super::TickComponent(DeltaTime, TickType, ThisTickFunction);
+
+	TSet<EMagicLeapHeadTrackingMapEvent> CurrentMapEvents;
+	if (UMagicLeapHMDFunctionLibrary::GetHeadTrackingMapEvents(CurrentMapEvents))
+	{
+		for (EMagicLeapHeadTrackingMapEvent MapEvent : CurrentMapEvents)
+		{
+			if (!PreviousMapEvents.Contains(MapEvent))
+			{
+				switch(MapEvent)
+				{
+					case EMagicLeapHeadTrackingMapEvent::Lost:
+					{
+						OnHeadTrackingLost.Broadcast();
+						HMDLostDelegate.Broadcast();
+						break;
+					}
+					case EMagicLeapHeadTrackingMapEvent::Recovered:
+					{
+						OnHeadTrackingRecovered.Broadcast();
+						HMDReconnectedDelegate.Broadcast();
+						break;
+					}
+					case EMagicLeapHeadTrackingMapEvent::RecoveryFailed:
+					{
+						OnHeadTrackingRecoveryFailed.Broadcast();
+						break;
+					}
+					case EMagicLeapHeadTrackingMapEvent::NewSession:
+					{
+						OnHeadTrackingNewSessionStarted.Broadcast();
+						HMDReconnectedDelegate.Broadcast();
+						break;
+					}
+				}
+			}
+		}
+
+		PreviousMapEvents = CurrentMapEvents;
+	}
+}
+
+void UMagicLeapHeadTrackingNotificationsComponent::BindToOnHeadTrackingLostEvent(const FVRNotificationsDelegate& EventDelegate)
+{
+	OnHeadTrackingLost = EventDelegate;
+}
+
+void UMagicLeapHeadTrackingNotificationsComponent::BindToOnHeadTrackingRecoveredEvent(const FVRNotificationsDelegate& EventDelegate)
+{
+	OnHeadTrackingRecovered = EventDelegate;
+}
+
+void UMagicLeapHeadTrackingNotificationsComponent::BindToOnHeadTrackingRecoveryFailedEvent(const FVRNotificationsDelegate& EventDelegate)
+{
+	OnHeadTrackingRecoveryFailed = EventDelegate;
+}
+
+void UMagicLeapHeadTrackingNotificationsComponent::BindToOnHeadTrackingNewSessionStartedEvent(const FVRNotificationsDelegate& EventDelegate)
+{
+	OnHeadTrackingNewSessionStarted = EventDelegate;
+}
+
+void UMagicLeapHeadTrackingNotificationsComponent::UnbindOnHeadTrackingLostEvent()
+{
+	OnHeadTrackingLost.Clear();
+}
+
+void UMagicLeapHeadTrackingNotificationsComponent::UnbindOnHeadTrackingRecoveredEvent()
+{
+	OnHeadTrackingRecovered.Clear();
+}
+
+void UMagicLeapHeadTrackingNotificationsComponent::UnbindOnHeadTrackingRecoveryFailedEvent()
+{
+	OnHeadTrackingRecoveryFailed.Clear();
+}
+
+void UMagicLeapHeadTrackingNotificationsComponent::UnbindOnHeadTrackingNewSessionStartedEvent()
+{
+	OnHeadTrackingNewSessionStarted.Clear();
+}
diff --git a/Engine/Plugins/Lumin/MagicLeap/Source/MagicLeap/Private/MagicLeapPool.h b/Engine/Plugins/Lumin/MagicLeap/Source/MagicLeap/Private/MagicLeapPool.h
--- a/Engine/Plugins/Lumin/MagicLeap/Source/MagicLeap/Private/MagicLeapPool.h
+++ b/Engine/Plugins/Lumin/MagicLeap/Source/MagicLeap/Private/MagicLeapPool.h
@@ -1,36 +1,4 @@
-// %BANNER_BEGIN%
-// ---------------------------------------------------------------------
-// %COPYRIGHT_BEGIN%
-//
-// Copyright (c) 2017 Magic Leap, Inc. (COMPANY) All Rights Reserved.
-// Magic Leap, Inc. Confidential and Proprietary
-//
-// NOTICE: All information contained herein is, and remains the property
-// of COMPANY. The intellectual and technical concepts contained herein
-// are proprietary to COMPANY and may be covered by U.S. and Foreign
-// Patents, patents in process, and are protected by trade secret or
-// copyright law. Dissemination of this information or reproduction of
-// this material is strictly forbidden unless prior written permission is
-// obtained from COMPANY. Access to the source code contained herein is
-// hereby forbidden to anyone except current COMPANY employees, managers
-// or contractors who have executed Confidentiality and Non-disclosure
-// agreements explicitly covering such access.
-//
-// The copyright notice above does not evidence any actual or intended
-// publication or disclosure of this source code, which includes
-// information that is confidential and/or proprietary, and is a trade
-// secret, of COMPANY. ANY REPRODUCTION, MODIFICATION, DISTRIBUTION,
-// PUBLIC PERFORMANCE, OR PUBLIC DISPLAY OF OR THROUGH USE OF THIS
-// SOURCE CODE WITHOUT THE EXPRESS WRITTEN CONSENT OF COMPANY IS
-// STRICTLY PROHIBITED, AND IN VIOLATION OF APPLICABLE LAWS AND
-// INTERNATIONAL TREATIES. THE RECEIPT OR POSSESSION OF THIS SOURCE
-// CODE AND/OR RELATED INFORMATION DOES NOT CONVEY OR IMPLY ANY RIGHTS
-// TO REPRODUCE, DISCLOSE OR DISTRIBUTE ITS CONTENTS, OR TO MANUFACTURE,
-// USE, OR SELL ANYTHING THAT IT MAY DESCRIBE, IN WHOLE OR IN PART.
-//
-// %COPYRIGHT_END%
-// --------------------------------------------------------------------
-// %BANNER_END%
+// Copyright 1998-2019 Epic Games, Inc. All Rights Reserved.
 
 #pragma once
 
diff --git a/Engine/Plugins/Lumin/MagicLeap/Source/MagicLeap/Private/MagicLeapUtils.cpp b/Engine/Plugins/Lumin/MagicLeap/Source/MagicLeap/Private/MagicLeapUtils.cpp
--- a/Engine/Plugins/Lumin/MagicLeap/Source/MagicLeap/Private/MagicLeapUtils.cpp
+++ b/Engine/Plugins/Lumin/MagicLeap/Source/MagicLeap/Private/MagicLeapUtils.cpp
@@ -7,9 +7,10 @@
 namespace MagicLeap
 {
 #if WITH_MLSDK
-	void ResetClipExtentsInfoArray(MLGraphicsClipExtentsInfoArray& UpdateInfoArray)
+	void ResetClipExtentsInfoArray(MLGraphicsClipExtentsInfoArrayEx& UpdateInfoArray)
 	{
-		UpdateInfoArray.num_virtual_cameras = 0;
+		MLGraphicsClipExtentsInfoArrayExInit(&UpdateInfoArray);
+
 		for (MLGraphicsClipExtentsInfo &ViewportInfo : UpdateInfoArray.virtual_camera_extents)
 		{
 			FMemory::Memcpy(ViewportInfo.projection.matrix_colmajor, MagicLeap::kIdentityMatColMajor, sizeof(MagicLeap::kIdentityMatColMajor));
diff --git a/Engine/Plugins/Lumin/MagicLeap/Source/MagicLeap/Private/MagicLeapUtils.h b/Engine/Plugins/Lumin/MagicLeap/Source/MagicLeap/Private/MagicLeapUtils.h
--- a/Engine/Plugins/Lumin/MagicLeap/Source/MagicLeap/Private/MagicLeapUtils.h
+++ b/Engine/Plugins/Lumin/MagicLeap/Source/MagicLeap/Private/MagicLeapUtils.h
@@ -1,3 +1,4 @@
+
 // Copyright 1998-2019 Epic Games, Inc. All Rights Reserved.
 
 #pragma once
@@ -9,7 +10,7 @@
 namespace MagicLeap
 {
 	// Clears the extents info and puts it into a safe initial state.
-	void ResetClipExtentsInfoArray(MLGraphicsClipExtentsInfoArray& UpdateInfoArray);
+	void ResetClipExtentsInfoArray(MLGraphicsClipExtentsInfoArrayEx& UpdateInfoArray);
 
 	// Clears the virtual camera info and puts it into a safe initial state.
 	void ResetVirtualCameraInfoArray(MLGraphicsVirtualCameraInfoArray& RenderInfoArray);
diff --git a/Engine/Plugins/Lumin/MagicLeap/Source/MagicLeap/Private/MagicLeapVulkanExtensions.cpp b/Engine/Plugins/Lumin/MagicLeap/Source/MagicLeap/Private/MagicLeapVulkanExtensions.cpp
--- a/Engine/Plugins/Lumin/MagicLeap/Source/MagicLeap/Private/MagicLeapVulkanExtensions.cpp
+++ b/Engine/Plugins/Lumin/MagicLeap/Source/MagicLeap/Private/MagicLeapVulkanExtensions.cpp
@@ -1,51 +1,11 @@
-// %BANNER_BEGIN%
-// ---------------------------------------------------------------------
-// %COPYRIGHT_BEGIN%
-//
-// Copyright (c) 201x Magic Leap, Inc. (COMPANY) All Rights Reserved.
-// Magic Leap, Inc. Confidential and Proprietary
-//
-// NOTICE: All information contained herein is, and remains the property
-// of COMPANY. The intellectual and technical concepts contained herein
-// are proprietary to COMPANY and may be covered by U.S. and Foreign
-// Patents, patents in process, and are protected by trade secret or
-// copyright law. Dissemination of this information or reproduction of
-// this material is strictly forbidden unless prior written permission is
-// obtained from COMPANY. Access to the source code contained herein is
-// hereby forbidden to anyone except current COMPANY employees, managers
-// or contractors who have executed Confidentiality and Non-disclosure
-// agreements explicitly covering such access.
-//
-// The copyright notice above does not evidence any actual or intended
-// publication or disclosure of this source code, which includes
-// information that is confidential and/or proprietary, and is a trade
-// secret, of COMPANY. ANY REPRODUCTION, MODIFICATION, DISTRIBUTION,
-// PUBLIC PERFORMANCE, OR PUBLIC DISPLAY OF OR THROUGH USE OF THIS
-// SOURCE CODE WITHOUT THE EXPRESS WRITTEN CONSENT OF COMPANY IS
-// STRICTLY PROHIBITED, AND IN VIOLATION OF APPLICABLE LAWS AND
-// INTERNATIONAL TREATIES. THE RECEIPT OR POSSESSION OF THIS SOURCE
-// CODE AND/OR RELATED INFORMATION DOES NOT CONVEY OR IMPLY ANY RIGHTS
-// TO REPRODUCE, DISCLOSE OR DISTRIBUTE ITS CONTENTS, OR TO MANUFACTURE,
-// USE, OR SELL ANYTHING THAT IT MAY DESCRIBE, IN WHOLE OR IN PART.
-//
-// %COPYRIGHT_END%
-// --------------------------------------------------------------------
-// %BANNER_END%
+// Copyright 1998-2019 Epic Games, Inc. All Rights Reserved.
 
 #include "MagicLeapVulkanExtensions.h"
+#include "Lumin/CAPIShims/LuminAPIRemote.h"
 
 #if !PLATFORM_LUMIN
-
 #include "MagicLeapGraphics.h"
-
-#if !PLATFORM_MAC
-#include <vulkan.h>
-#endif // !PLATFORM_MAC
-#if WITH_MLSDK
-#include <ml_remote.h>
-#endif // WITH_MLSDK
 #include "AppFramework.h"
-
 #endif // !PLATFORM_LUMIN
 
 #include "MagicLeapHelperVulkan.h"
@@ -67,6 +27,9 @@ bool FMagicLeapVulkanExtensions::GetVulkanInstanceExtensionsRequired(TArray<cons
 	{
 		ImpPtr.Reset(new Implementation);
 	}
+#if PLATFORM_LUMIN
+	return FMagicLeapHelperVulkan::GetVulkanInstanceExtensionsRequired(Out);
+#else
 #if (PLATFORM_WINDOWS && WITH_MLSDK)
 	// Interrogate the extensions we need for MLRemote.
 	TArray<VkExtensionProperties> Extensions;
@@ -97,6 +60,7 @@ bool FMagicLeapVulkanExtensions::GetVulkanInstanceExtensionsRequired(TArray<cons
 	}
 #endif //(PLATFORM_WINDOWS && WITH_MLSDK)
 	return true;
+#endif // PLATFORM_LUMIN
 }
 
 bool FMagicLeapVulkanExtensions::GetVulkanDeviceExtensionsRequired(struct VkPhysicalDevice_T *pPhysicalDevice, TArray<const ANSICHAR*>& Out)
@@ -131,7 +95,7 @@ bool FMagicLeapVulkanExtensions::GetVulkanDeviceExtensionsRequired(struct VkPhys
 				return false;
 			}
 		}
-}
+	}
 	for (auto & Extension : ImpPtr->DeviceExtensions)
 	{
 		Out.Add(Extension.extensionName);
diff --git a/Engine/Plugins/Lumin/MagicLeap/Source/MagicLeap/Private/MagicLeapVulkanExtensions.h b/Engine/Plugins/Lumin/MagicLeap/Source/MagicLeap/Private/MagicLeapVulkanExtensions.h
--- a/Engine/Plugins/Lumin/MagicLeap/Source/MagicLeap/Private/MagicLeapVulkanExtensions.h
+++ b/Engine/Plugins/Lumin/MagicLeap/Source/MagicLeap/Private/MagicLeapVulkanExtensions.h
@@ -1,36 +1,4 @@
-// %BANNER_BEGIN%
-// ---------------------------------------------------------------------
-// %COPYRIGHT_BEGIN%
-//
-// Copyright (c) 201x Magic Leap, Inc. (COMPANY) All Rights Reserved.
-// Magic Leap, Inc. Confidential and Proprietary
-//
-// NOTICE: All information contained herein is, and remains the property
-// of COMPANY. The intellectual and technical concepts contained herein
-// are proprietary to COMPANY and may be covered by U.S. and Foreign
-// Patents, patents in process, and are protected by trade secret or
-// copyright law. Dissemination of this information or reproduction of
-// this material is strictly forbidden unless prior written permission is
-// obtained from COMPANY. Access to the source code contained herein is
-// hereby forbidden to anyone except current COMPANY employees, managers
-// or contractors who have executed Confidentiality and Non-disclosure
-// agreements explicitly covering such access.
-//
-// The copyright notice above does not evidence any actual or intended
-// publication or disclosure of this source code, which includes
-// information that is confidential and/or proprietary, and is a trade
-// secret, of COMPANY. ANY REPRODUCTION, MODIFICATION, DISTRIBUTION,
-// PUBLIC PERFORMANCE, OR PUBLIC DISPLAY OF OR THROUGH USE OF THIS
-// SOURCE CODE WITHOUT THE EXPRESS WRITTEN CONSENT OF COMPANY IS
-// STRICTLY PROHIBITED, AND IN VIOLATION OF APPLICABLE LAWS AND
-// INTERNATIONAL TREATIES. THE RECEIPT OR POSSESSION OF THIS SOURCE
-// CODE AND/OR RELATED INFORMATION DOES NOT CONVEY OR IMPLY ANY RIGHTS
-// TO REPRODUCE, DISCLOSE OR DISTRIBUTE ITS CONTENTS, OR TO MANUFACTURE,
-// USE, OR SELL ANYTHING THAT IT MAY DESCRIBE, IN WHOLE OR IN PART.
-//
-// %COPYRIGHT_END%
-// --------------------------------------------------------------------
-// %BANNER_END%
+// Copyright 1998-2019 Epic Games, Inc. All Rights Reserved.
 
 #pragma once
 
diff --git a/Engine/Plugins/Lumin/MagicLeap/Source/MagicLeap/Private/MagicLeapXRCamera.cpp b/Engine/Plugins/Lumin/MagicLeap/Source/MagicLeap/Private/MagicLeapXRCamera.cpp
--- a/Engine/Plugins/Lumin/MagicLeap/Source/MagicLeap/Private/MagicLeapXRCamera.cpp
+++ b/Engine/Plugins/Lumin/MagicLeap/Source/MagicLeap/Private/MagicLeapXRCamera.cpp
@@ -22,11 +22,9 @@ void FMagicLeapXRCamera::PreRenderView_RenderThread(FRHICommandListImmediate& RH
 		const FTrackingFrame& Frame = MagicLeapSystem.GetCurrentFrame();
 
 		// update to use render projection matrix
-		// #todo: Roll UpdateProjectionMatrix into UpdateViewMatrix?
-		FMatrix RenderInfoProjectionMatrix = MagicLeap::ToFMatrix(Frame.RenderInfoArray.virtual_cameras[EyeIdx].projection);
-
 		// Set the near clipping plane to GNearClippingPlane which is clamped to the minimum value allowed for the device. (ref: MLGraphicsGetRenderTargets())
-		RenderInfoProjectionMatrix.M[3][2] = GNearClippingPlane;
+		// #todo: Roll UpdateProjectionMatrix into UpdateViewMatrix?
+		FMatrix RenderInfoProjectionMatrix = MagicLeap::ToUEProjectionMatrix(Frame.FrameInfo.virtual_camera_info_array.virtual_cameras[EyeIdx].projection, GNearClippingPlane);
 
 		View.UpdateProjectionMatrix(RenderInfoProjectionMatrix);
 	}
