Copying //UE4/Partner-MagicLeap-Staging to Dev-VR (//UE4/Dev-VR)

diff --git a/Engine/Plugins/Lumin/MagicLeap/Source/MagicLeap/Private/MeshTrackerComponent.cpp b/Engine/Plugins/Lumin/MagicLeap/Source/MagicLeap/Private/MeshTrackerComponent.cpp
--- a/Engine/Plugins/Lumin/MagicLeap/Source/MagicLeap/Private/MeshTrackerComponent.cpp
+++ b/Engine/Plugins/Lumin/MagicLeap/Source/MagicLeap/Private/MeshTrackerComponent.cpp
@@ -1,14 +1,12 @@
 // Copyright 1998-2019 Epic Games, Inc. All Rights Reserved.
 
 #include "MeshTrackerComponent.h"
-#include "IMagicLeapHMD.h"
 #include "MagicLeapHMD.h"
 #include "AppFramework.h"
 #include "MagicLeapMath.h"
 #include "AppEventHandler.h"
+#include "MagicLeapHMDFunctionLibrary.h"
 
-#include "KismetProceduralMeshLibrary.h"
-#include "ProceduralMeshComponent.h"
 #include "Components/BoxComponent.h"
 #include "Engine/Engine.h"
 #include "HeadMountedDisplayFunctionLibrary.h"
@@ -16,52 +14,49 @@
 #include "Editor.h"
 #endif
 #include "MRMeshComponent.h"
-#include "MagicLeapPluginUtil.h" // for ML_INCLUDES_START/END
+#include "Lumin/CAPIShims/LuminAPIMeshing.h"
 
-#if WITH_MLSDK
-ML_INCLUDES_START
-#include <ml_meshing2.h>
-ML_INCLUDES_END
-#endif //WITH_MLSDK
+//#define DEBUG_MESH_BLOCK_ADD_REMOVE
+//#define DEBUG_MESH_REQUEST_AND_RESPONSE
 
 #if WITH_MLSDK
 // TODO: Don't rely on the size being same.
 static_assert(sizeof(FGuid) == sizeof(MLCoordinateFrameUID), "Size of FGuid should be same as MLCoordinateFrameUID. TODO: Don't rely on the size being same.");
 
 // Map an Unreal meshing LOD to the corresponding ML meshing LOD
-FORCEINLINE MLMeshingLOD UnrealToML_MeshLOD(EMeshLOD UnrealMeshLOD)
+FORCEINLINE MLMeshingLOD UnrealToML_MeshLOD(EMagicLeapMeshLOD UnrealMeshLOD)
 {
 	switch (UnrealMeshLOD)
 	{
-		case EMeshLOD::Minimum:
+		case EMagicLeapMeshLOD::Minimum:
 			return MLMeshingLOD_Minimum;
-		case EMeshLOD::Medium:
+		case EMagicLeapMeshLOD::Medium:
 			return MLMeshingLOD_Medium;
-		case EMeshLOD::Maximum:
+		case EMagicLeapMeshLOD::Maximum:
 			return MLMeshingLOD_Maximum;
 	}
 	check(false);
 	return MLMeshingLOD_Minimum;
 }
 
-EMeshState MLToUEMeshState(MLMeshingMeshState MLMeshState)
+EMagicLeapMeshState MLToUEMeshState(MLMeshingMeshState MLMeshState)
 {
 	switch (MLMeshState)
 	{
 		case MLMeshingMeshState_New:
-			return EMeshState::New;
+			return EMagicLeapMeshState::New;
 		case MLMeshingMeshState_Updated:
-			return EMeshState::Updated;
+			return EMagicLeapMeshState::Updated;
 		case MLMeshingMeshState_Deleted:
-			return EMeshState::Deleted;
+			return EMagicLeapMeshState::Deleted;
 		case MLMeshingMeshState_Unchanged:
-			return EMeshState::Unchanged;
+			return EMagicLeapMeshState::Unchanged;
 	}
 	check(false);
-	return EMeshState::Unchanged;
+	return EMagicLeapMeshState::Unchanged;
 }
 
-void MLToUnrealBlockInfo(const MLMeshingBlockInfo& MLBlockInfo, const FTransform& TrackingToWorld, float WorldToMetersScale, FMeshBlockInfo& UEBlockInfo)
+void MLToUnrealBlockInfo(const MLMeshingBlockInfo& MLBlockInfo, const FTransform& TrackingToWorld, float WorldToMetersScale, FMagicLeapMeshBlockInfo& UEBlockInfo)
 {
 	FMemory::Memcpy(&UEBlockInfo.BlockID, &MLBlockInfo.id, sizeof(MLCoordinateFrameUID));
 
@@ -76,29 +71,37 @@ void MLToUnrealBlockInfo(const MLMeshingBlockInfo& MLBlockInfo, const FTransform
 	BlockTransform = BlockTransform * TrackingToWorld;
 	UEBlockInfo.BlockPosition = BlockTransform.GetLocation();
 	UEBlockInfo.BlockOrientation = BlockTransform.Rotator();
-	UEBlockInfo.BlockDimensions = MagicLeap::ToFVectorExtents(MLBlockInfo.extents.extents, WorldToMetersScale);
+
+	// Splat the OBB to an AABB
+	FMatrix AbsWorldMatrix(BlockTransform.ToMatrixNoScale());
+	AbsWorldMatrix.SetAxis(3, FVector::ZeroVector);
+	for (int32 R = 0; R < 3; ++R)
+	{
+		for (int32 C = 0; C < 3; ++C)
+		{
+			AbsWorldMatrix.M[R][C] = FMath::Abs(AbsWorldMatrix.M[R][C]);
+		}
+	}
+
+	// The extents returned are 'full' extents - width, height, and depth. UE4 boxes are 'half' extents
+	UEBlockInfo.BlockDimensions = AbsWorldMatrix.
+		TransformVector(MagicLeap::ToFVectorExtents(MLBlockInfo.extents.extents, WorldToMetersScale * 0.5f));
 
 	UEBlockInfo.Timestamp = FTimespan::FromMicroseconds(MLBlockInfo.timestamp / 1000.0);
 	UEBlockInfo.BlockState = MLToUEMeshState(MLBlockInfo.state);
 }
 
-void UnrealToMLBlockRequest(const FMeshBlockRequest& UEBlockRequest, MLMeshingBlockRequest& MLBlockRequest)
+void UnrealToMLBlockRequest(const FMagicLeapMeshBlockRequest& UEBlockRequest, MLMeshingBlockRequest& MLBlockRequest)
 {
 	FMemory::Memcpy(&MLBlockRequest.id, &UEBlockRequest.BlockID, sizeof(MLCoordinateFrameUID));
 	MLBlockRequest.level = UnrealToML_MeshLOD(UEBlockRequest.LevelOfDetail);
 }
 #endif //WITH_MLSDK
 
-class FMeshTrackerImpl : public MagicLeap::IAppEventHandler
+class FMagicLeapMeshTrackerImpl : public MagicLeap::IAppEventHandler
 {
 public:
-	FMeshTrackerImpl()
-#if WITH_MLSDK
-		: MeshTracker(ML_INVALID_HANDLE)
-		, MeshBrickIndex(0)
-		, CurrentMeshInfoRequest(ML_INVALID_HANDLE)
-		, CurrentMeshRequest(ML_INVALID_HANDLE)
-#endif //WITH_MLSDK
+	FMagicLeapMeshTrackerImpl()
 	{
 	};
 
@@ -114,13 +117,10 @@ public:
 		{
 			const FAppFramework& AppFramework = static_cast<FMagicLeapHMD *>
 				(GEngine->XRSystem->GetHMDDevice())->GetAppFrameworkConst();
-			if (AppFramework.IsInitialized())
-			{
-				WorldToMetersScale = AppFramework.GetWorldToMetersScale();
-			}
+			WorldToMetersScale = AppFramework.GetWorldToMetersScale();
 		}
 
-		if (MeshTrackerComponent.MeshType == EMeshType::PointCloud)
+		if (MeshTrackerComponent.MeshType == EMagicLeapMeshType::PointCloud)
 		{
 			Settings.flags |= MLMeshingFlags_PointCloud;
 		}
@@ -158,43 +158,55 @@ public:
 	{
 	}
 
+	void OnClear()
+	{
+		MeshBrickIndex = 0;
+		GuidToBrickId.Empty();
+
+		for (auto& PendingMeshBrick : PendingMeshBricksByBrickId)
+		{
+			PendingMeshBrick.Value->Recycle(PendingMeshBrick.Value);
+		}
+		PendingMeshBricksByBrickId.Empty();
+	}
+
 public:
 #if WITH_MLSDK
 	// Handle to ML mesh tracker
-	MLHandle MeshTracker;
-
-	// Next ID for bricks created with MR Mesh
-	int32 MeshBrickIndex;
+	MLHandle MeshTracker = ML_INVALID_HANDLE;
 
 	// Handle to ML mesh info request
-	MLHandle CurrentMeshInfoRequest;
+	MLHandle CurrentMeshInfoRequest = ML_INVALID_HANDLE;
 
 	// Handle to ML mesh request
-	MLHandle CurrentMeshRequest;
+	MLHandle CurrentMeshRequest = ML_INVALID_HANDLE;
 
 	// Current ML meshing settings
 	MLMeshingSettings CurrentMeshSettings;
 
 	// List of ML mesh block IDs and states
 	TArray<MLMeshingBlockRequest> MeshBlockRequests;
-
-	// The set of meshing blocks that have been processed or removed between queries
-	TSet<FGuid> MeshBrickCache;
-
 #endif //WITH_MLSDK
 
-	FMLTrackingMeshInfo LastestMeshInfo;
-	TArray<FMeshBlockRequest> UEMeshBlockRequests;
-	TScriptInterface<IMeshBlockSelectorInterface> BlockSelector;
+	struct FMLTrackingInfoImpl : FMagicLeapTrackingMeshInfo
+	{
+		// Maps the FMagicLeapTrackingMeshInfo block GUIDs to their entry in the BlockData array
+		TMap<FGuid, FMagicLeapMeshBlockInfo> BlockInfoByGuid;
+	};
+
+	FMLTrackingInfoImpl LatestMeshInfo;
+	TArray<FMagicLeapMeshBlockRequest> UEMeshBlockRequests;
+	TScriptInterface<IMagicLeapMeshBlockSelectorInterface> BlockSelector;
 
-	// Keep a copy of the mesh data here.  MRMeshComponent will use it from the game and render thread.
+	// Keep a copy of the mesh data here. MRMeshComponent will use it from the game and render thread.
 	struct FMLCachedMeshData
 	{
 		typedef TSharedPtr<FMLCachedMeshData, ESPMode::ThreadSafe> SharedPtr;
 		
-		FMeshTrackerImpl* Owner = nullptr;
+		FMagicLeapMeshTrackerImpl* Owner = nullptr;
 		
 		FGuid BlockID;
+		uint64 BrickID;
 		TArray<FVector> OffsetVertices;
 		TArray<FVector> WorldVertices;
 		TArray<uint32> Triangles;
@@ -204,13 +216,14 @@ public:
 		TArray<FPackedNormal> Tangents;
 		TArray<float> Confidence;
 
+		FMagicLeapMeshBlockInfo BlockInfo;
+
 		void Recycle(SharedPtr& MeshData)
 		{
 			check(Owner);
-			FMeshTrackerImpl* TempOwner = Owner;
-			Owner = nullptr;
-	
+
 			BlockID.Invalidate();
+			BrickID = INT64_MAX;
 			OffsetVertices.Reset();
 			WorldVertices.Reset();
 			Triangles.Reset();
@@ -219,17 +232,23 @@ public:
 			VertexColors.Reset();
 			Tangents.Reset();
 			Confidence.Reset();
-	
-			TempOwner->FreeMeshDataCache(MeshData);
+
+			Owner->FreeMeshDataCache(MeshData);
+			Owner = nullptr;
 		}
 
-		void Init(FMeshTrackerImpl* InOwner)
+		void Init(FMagicLeapMeshTrackerImpl* InOwner)
 		{
 			check(!Owner);
 			Owner = InOwner;
 		}
 	};
 
+	// When load-balancing is enabled (BricksPerFrame > 0) this map contains mesh data
+	// pending submission to MR Mesh
+	TMap<uint64, FMLCachedMeshData::SharedPtr> PendingMeshBricksByBrickId;
+
+	FDelegateHandle OnClearDelegateHandle;
 
 	// This receipt will be kept in the FSendBrickDataArgs to ensure the cached data outlives MRMeshComponent use of it.
 	class FMeshTrackerComponentBrickDataReceipt : public IMRMesh::FBrickDataReceipt
@@ -247,7 +266,7 @@ public:
 		FMLCachedMeshData::SharedPtr CachedMeshData;
 	};
 	
-	FMLCachedMeshData::SharedPtr AquireMeshDataCache()
+	FMLCachedMeshData::SharedPtr AcquireMeshDataCache()
 	{
 		if (FreeCachedMeshDatas.Num() > 0)
 		{
@@ -293,6 +312,7 @@ public:
 				return false;
 			}
 
+			GuidToBrickId.Empty();
 			MeshBrickIndex = 0;
 		}
 #endif //WITH_MLSDK
@@ -318,7 +338,7 @@ public:
 					{
 						UE_LOG(LogMagicLeap, Log,
 							TEXT("MLMeshingSettings change caused a clear"));
-						MeshTrackerComponent.MRMesh->Clear();
+						MeshTrackerComponent.MRMesh->Clear(); // Will call our OnClear
 					}
 				}
 
@@ -349,7 +369,23 @@ public:
 #endif //WITH_MLSDK
 	}
 
+	uint64 *GetMrMeshIdFromMeshGuid(const FGuid& meshGuid, bool addIfNotFound)
+	{
+		auto MeshBrickId = GuidToBrickId.Find(meshGuid);
+		if (MeshBrickId == nullptr && addIfNotFound)
+		{
+			GuidToBrickId.Add(meshGuid, MeshBrickIndex++);
+			MeshBrickId = GuidToBrickId.Find(meshGuid);
+		}
+		return MeshBrickId;
+	}
+
 private:
+	// Maps system 128-bit mesh GUIDs to 64-bit brick indices
+	TMap<FGuid, uint64> GuidToBrickId;
+
+	// Next 64-bit brick index
+	uint64 MeshBrickIndex = 0;
 
 	// A free list to recycle the CachedMeshData instances.  
 	TArray<FMLCachedMeshData::SharedPtr> CachedMeshDatas;
@@ -360,7 +396,7 @@ private:
 UMeshTrackerComponent::UMeshTrackerComponent()
 	: VertexColorFromConfidenceZero(FLinearColor::Red)
 	, VertexColorFromConfidenceOne(FLinearColor::Blue)
-	, Impl(new FMeshTrackerImpl())
+	, Impl(new FMagicLeapMeshTrackerImpl())
 {
 
 	// Make sure this component ticks
@@ -378,7 +414,7 @@ UMeshTrackerComponent::UMeshTrackerComponent()
 	BoundingVolume->SetCollisionObjectType(ECollisionChannel::ECC_WorldDynamic);
 	BoundingVolume->SetGenerateOverlapEvents(false);
 	// Recommended default box extents for meshing - 10m (5m radius)
-	BoundingVolume->SetBoxExtent(FVector(1000, 1000, 1000), false);
+	BoundingVolume->SetBoxExtent(FVector(500.0f, 500.0f, 500.0f), false);
 
 	BlockVertexColors.Add(FColor::Blue);
 	BlockVertexColors.Add(FColor::Red);
@@ -424,6 +460,7 @@ void UMeshTrackerComponent::ConnectMRMesh(UMRMeshComponent* InMRMeshPtr)
 	{
 		InMRMeshPtr->SetConnected(true);
 		MRMesh = InMRMeshPtr;
+		Impl->OnClearDelegateHandle = MRMesh->OnClear().AddRaw(Impl, &FMagicLeapMeshTrackerImpl::OnClear);
 	}
 }
 
@@ -446,6 +483,12 @@ void UMeshTrackerComponent::DisconnectMRMesh(class UMRMeshComponent* InMRMeshPtr
 	{
 		check(MRMesh->IsConnected());
 		MRMesh->SetConnected(false);
+
+		if (Impl->OnClearDelegateHandle.IsValid())
+		{
+			MRMesh->OnClear().Remove(Impl->OnClearDelegateHandle);
+			Impl->OnClearDelegateHandle.Reset();
+		}
 	}
 	MRMesh = nullptr;
 }
@@ -485,10 +528,6 @@ void UMeshTrackerComponent::TickComponent(float DeltaTime, enum ELevelTick TickT
 
 	const FAppFramework& AppFramework = 
 		static_cast<FMagicLeapHMD *>(GEngine->XRSystem->GetHMDDevice())->GetAppFrameworkConst();
-	if (!AppFramework.IsInitialized())
-	{
-		return;
-	}
 
 	if (!Impl->Create(*this))
 	{
@@ -514,6 +553,20 @@ void UMeshTrackerComponent::TickComponent(float DeltaTime, enum ELevelTick TickT
 		UE_LOG(LogMagicLeap, Log, TEXT("MLMeshingSettings changed on the fly"));
 	}
 
+	TSet<EMagicLeapHeadTrackingMapEvent> MapEvents;
+	bool bHasMapEvents = UMagicLeapHMDFunctionLibrary::GetHeadTrackingMapEvents(MapEvents);
+	if (bHasMapEvents)
+	{
+		for (EMagicLeapHeadTrackingMapEvent MapEvent : MapEvents)
+		{
+			if (MapEvent == EMagicLeapHeadTrackingMapEvent::NewSession)
+			{
+				// Clear existing meshes if a new headpose session has started.
+				MRMesh->Clear();
+			}
+		}
+	}
+
 	const float WorldToMetersScale = AppFramework.GetWorldToMetersScale();
 
 	// Make sure MR Mesh is at 0,0,0 (verts received from ML meshing are in tracking space)
@@ -529,25 +582,99 @@ void UMeshTrackerComponent::TickComponent(float DeltaTime, enum ELevelTick TickT
 				RequestMesh();
 			}
 		}
+
+		if (BricksPerFrame > 0)
+		{
+			// Load-balancing of MR Mesh brick creation. It is recommended that applications use the GetNumQueuedBlockUpdates
+			// function and experimentation to fine-tune this number for their application.
+			auto PendingIter = Impl->PendingMeshBricksByBrickId.CreateConstIterator();
+			int32 numBricksProcessed;
+			for (numBricksProcessed = 0; numBricksProcessed < BricksPerFrame && PendingIter; ++numBricksProcessed)
+			{
+				auto CachedMeshData = PendingIter.Value();
+				++PendingIter; // Assuming this needs to be done before removing the entry it points to
+
+#ifdef DEBUG_MESH_BLOCK_ADD_REMOVE
+				if (CachedMeshData->WorldVertices.Num() > 0)
+				{
+					UE_LOG(LogMagicLeap,
+						Log,
+						TEXT("UMeshTrackerComponent: ADDING/UPDATING brick %s"),
+						*(CachedMeshData->BlockID.ToString()));
+				}
+				else
+				{
+					UE_LOG(LogMagicLeap,
+						Log,
+						TEXT("UMeshTrackerComponent: REMOVING brick %s"),
+						*(CachedMeshData->BlockID.ToString()));
+				}
+#endif //DEBUG_MESH_BLOCK_ADD_REMOVE
+				// Broadcast that a mesh was updated
+				if (OnMeshTrackerUpdated.IsBound())
+				{
+					// Hack because blueprints don't support uint32.
+					TArray<int32> Triangles(reinterpret_cast<const int32*>(CachedMeshData->
+						Triangles.GetData()), CachedMeshData->Triangles.Num());
+					OnMeshTrackerUpdated.Broadcast(CachedMeshData->BlockID,
+						CachedMeshData->OffsetVertices, Triangles, CachedMeshData->Normals,
+						CachedMeshData->Confidence);
+				}
+
+				// Remove it from the pending list
+				Impl->PendingMeshBricksByBrickId.Remove(CachedMeshData->BrickID);
+
+				if (MeshType != EMagicLeapMeshType::PointCloud)
+				{
+					// Create/update brick
+					static_cast<IMRMesh*>(MRMesh)->SendBrickData(IMRMesh::FSendBrickDataArgs
+						{
+							TSharedPtr<IMRMesh::FBrickDataReceipt, ESPMode::ThreadSafe>
+							(new FMagicLeapMeshTrackerImpl::FMeshTrackerComponentBrickDataReceipt(CachedMeshData)),
+							CachedMeshData->BrickID,
+							CachedMeshData->WorldVertices,
+							CachedMeshData->UV0,
+							CachedMeshData->Tangents,
+							CachedMeshData->VertexColors,
+							CachedMeshData->Triangles,
+							FBox(CachedMeshData->BlockInfo.BlockPosition - CachedMeshData->BlockInfo.BlockDimensions,
+								 CachedMeshData->BlockInfo.BlockPosition + CachedMeshData->BlockInfo.BlockDimensions)
+						});
+				}
+				else
+				{
+					// Discard
+					CachedMeshData->Recycle(CachedMeshData);
+				}
+			}
+		}
+		else
+		{
+			// Clear pending queries as it will not be drained
+			Impl->PendingMeshBricksByBrickId.Reset();
+		}
+	}
+	else
+	{
+		// Clear pending queries as it will not be drained
+		Impl->PendingMeshBricksByBrickId.Reset();
 	}
 #endif //WITH_MLSDK
+
 }
 
-void UMeshTrackerComponent::SelectMeshBlocks_Implementation(const FMLTrackingMeshInfo& NewMeshInfo, TArray<FMeshBlockRequest>& RequestedMesh)
+void UMeshTrackerComponent::SelectMeshBlocks_Implementation(const FMagicLeapTrackingMeshInfo& NewMeshInfo, TArray<FMagicLeapMeshBlockRequest>& RequestedMesh)
 {
-	for (const FMeshBlockInfo& BlockInfo : NewMeshInfo.BlockData)
+	for (const FMagicLeapMeshBlockInfo& BlockInfo : NewMeshInfo.BlockData)
 	{
-		if (BlockInfo.BlockState == EMeshState::New || BlockInfo.BlockState == EMeshState::Unchanged)
-		{
-			FMeshBlockRequest BlockRequest;
-			BlockRequest.BlockID = BlockInfo.BlockID;
-			BlockRequest.LevelOfDetail = LevelOfDetail;
-			RequestedMesh.Add(BlockRequest);
-		}
+		FMagicLeapMeshBlockRequest BlockRequest;
+		BlockRequest.BlockID = BlockInfo.BlockID;
+		BlockRequest.LevelOfDetail = LevelOfDetail;
+		RequestedMesh.Add(BlockRequest);
 	}
 }
 
-void UMeshTrackerComponent::ConnectBlockSelector(TScriptInterface<IMeshBlockSelectorInterface> Selector)
+void UMeshTrackerComponent::ConnectBlockSelector(TScriptInterface<IMagicLeapMeshBlockSelectorInterface> Selector)
 {
 	if (Impl != nullptr)
 	{
@@ -555,14 +682,14 @@ void UMeshTrackerComponent::ConnectBlockSelector(TScriptInterface<IMeshBlockSele
 		// Since the InterfacePointer is null for it's blueprint implementors, bool() operator gives us the wrong result for checking if interface is valid.
 		if (Selector.GetObject() != nullptr)
 		{
-			// If called via C++, Selector might have been created manually and not implement IMeshBlockSelectorInterface.
-			if (Selector.GetObject()->GetClass()->ImplementsInterface(UMeshBlockSelectorInterface::StaticClass()))
+			// If called via C++, Selector might have been created manually and not implement IMagicLeapMeshBlockSelectorInterface.
+			if (Selector.GetObject()->GetClass()->ImplementsInterface(UMagicLeapMeshBlockSelectorInterface::StaticClass()))
 			{
 				Impl->BlockSelector = Selector;
 			}
 			else
 			{
-				UE_LOG(LogMagicLeap, Warning, TEXT("Selector %s does not implement IMeshBlockSelectorInterface. Using default block selector from MeshTrackerComponent."), *(Selector.GetObject()->GetFName().ToString()));
+				UE_LOG(LogMagicLeap, Warning, TEXT("Selector %s does not implement IMagicLeapMeshBlockSelectorInterface. Using default block selector from MeshTrackerComponent."), *(Selector.GetObject()->GetFName().ToString()));
 				Impl->BlockSelector = this;	
 			}
 		}
@@ -582,6 +709,11 @@ void UMeshTrackerComponent::DisconnectBlockSelector()
 	}
 }
 
+int32 UMeshTrackerComponent::GetNumQueuedBlockUpdates()
+{
+	return Impl->PendingMeshBricksByBrickId.Num();
+}
+
 void UMeshTrackerComponent::BeginDestroy()
 {
 	if (MRMesh != nullptr)
@@ -616,7 +748,8 @@ void UMeshTrackerComponent::RequestMeshInfo()
 		MLMeshingExtents Extents = {};
 		Extents.center = MagicLeap::ToMLVector(Impl->BoundsCenter, WorldToMetersScale);
 		Extents.rotation = MagicLeap::ToMLQuat(Impl->BoundsRotation);
-		Extents.extents = MagicLeap::ToMLVectorExtents(BoundingVolume->GetScaledBoxExtent(), WorldToMetersScale);
+		// The C-API extents are 'full' extents - width, height, and depth. UE4 boxes are 'half' extents.
+		Extents.extents = MagicLeap::ToMLVectorExtents(BoundingVolume->GetScaledBoxExtent() * 2, WorldToMetersScale);
 
 		auto Result = MLMeshingRequestMeshInfo(Impl->MeshTracker, &Extents, &Impl->CurrentMeshInfoRequest);
 		if (MLResult_Ok != Result)
@@ -625,7 +758,7 @@ void UMeshTrackerComponent::RequestMeshInfo()
 			Impl->CurrentMeshInfoRequest = ML_INVALID_HANDLE;
 		}
 	}
-#endif
+#endif // WITH_MLSDK
 }
 
 bool UMeshTrackerComponent::GetMeshInfoResult()
@@ -646,7 +779,7 @@ bool UMeshTrackerComponent::GetMeshInfoResult()
 			// Just silently wait for pending result
 			if (MLResult_Pending != Result)
 			{
-				UE_LOG(LogMagicLeap, Error, 
+				UE_LOG(LogMagicLeap, Error,
 					TEXT("MLMeshingGetMeshInfoResult failed: %s."), UTF8_TO_TCHAR(MLGetResultString(Result)));
 				Impl->CurrentMeshInfoRequest = ML_INVALID_HANDLE;
 				return true;
@@ -656,55 +789,113 @@ bool UMeshTrackerComponent::GetMeshInfoResult()
 		else
 		{
 			// Clear our stored block requests
-			Impl->LastestMeshInfo.BlockData.Empty(MeshInfo.data_count);
-			Impl->LastestMeshInfo.BlockData.AddUninitialized(MeshInfo.data_count);
-			Impl->LastestMeshInfo.Timestamp = FTimespan::FromMicroseconds(MeshInfo.timestamp / 1000.0);
+			Impl->LatestMeshInfo.BlockInfoByGuid.Empty(MeshInfo.data_count);
+			Impl->LatestMeshInfo.BlockData.Empty(MeshInfo.data_count);
+			Impl->LatestMeshInfo.BlockData.AddUninitialized(MeshInfo.data_count);
+			Impl->LatestMeshInfo.Timestamp = FTimespan::FromMicroseconds(MeshInfo.timestamp / 1000.0);
 
-			for (uint32_t MeshInfoIndex = 0; MeshInfoIndex < MeshInfo.data_count; ++ MeshInfoIndex)
+			uint32_t MeshIndex = 0;
+			for (uint32_t DataIndex = 0; DataIndex < MeshInfo.data_count; ++DataIndex)
 			{
-				const auto &MeshInfoData = MeshInfo.data[MeshInfoIndex];
-
-				// TODO: right now we are adding even the deleted and unchanged blocks here. we probably only need to add new and unchanged blocks.
-				MLToUnrealBlockInfo(MeshInfoData, UHeadMountedDisplayFunctionLibrary::GetTrackingToWorldTransform(this), WorldToMetersScale, Impl->LastestMeshInfo.BlockData[MeshInfoIndex]);
+				const auto &MeshInfoData = MeshInfo.data[DataIndex];
 
 				switch (MeshInfoData.state)
 				{
-					// TODO: maybe do this only when the rest of the mesh is updated i.e. success from MLMeshingGetMeshResult()
-					case MLMeshingMeshState_Deleted:
+					case MLMeshingMeshState_Unchanged:
+					case MLMeshingMeshState_New:
+					case MLMeshingMeshState_Updated:
 					{
-						const auto& BlockID = FGuid(MeshInfoData.id.data[0], MeshInfoData.id.data[0] >> 32, MeshInfoData.id.data[1], MeshInfoData.id.data[1] >> 32);
+						auto BlockID = FGuid(MeshInfoData.id.data[0],
+											 MeshInfoData.id.data[0] >> 32,
+											 MeshInfoData.id.data[1],
+											 MeshInfoData.id.data[1] >> 32);
 
-						// Delete the brick and its cache entry
-						if (Impl->MeshBrickCache.Contains(BlockID))
+						if (MeshInfoData.state == MLMeshingMeshState_Unchanged)
 						{
-
-							if (MeshType != EMeshType::PointCloud)
+							// Make sure we have actually received this brick before considering
+							// it unchanged
+							auto MeshBrickId = Impl->GetMrMeshIdFromMeshGuid(BlockID, false);
+							if (MeshBrickId != nullptr)
 							{
-								const static TArray<FVector> EmptyVertices;
-								const static TArray<FVector2D> EmptyUVs;
-								const static TArray<FPackedNormal> EmptyTangents;
-								const static TArray<FColor> EmptyVertexColors;
-								const static TArray<uint32> EmptyTriangles;
-								static_cast<IMRMesh*>(MRMesh)->SendBrickData(IMRMesh::FSendBrickDataArgs
-									{
-										nullptr,
-										(static_cast<uint64>(BlockID.B) << 32 | BlockID.A),
-										EmptyVertices,
-										EmptyUVs,
-										EmptyTangents,
-										EmptyVertexColors,
-										EmptyTriangles
-									}
-								);
+								break;
 							}
 
+#ifdef DEBUG_MESH_REQUEST_AND_RESPONSE
+							UE_LOG(LogMagicLeap,
+								Log,
+								TEXT("UMeshTrackerComponent: Received 'unchanged' event for unseen block %s"),
+								*(BlockID.ToString()));
+#endif
+						}
+
+						// Convert and add the block info, and map it by GUID
+						auto& BlockInfo = Impl->LatestMeshInfo.BlockData[MeshIndex++];
+						MLToUnrealBlockInfo(MeshInfoData,
+							UHeadMountedDisplayFunctionLibrary::GetTrackingToWorldTransform(this),
+							WorldToMetersScale,
+							BlockInfo);
+
+						// It's new to us
+						BlockInfo.BlockState = EMagicLeapMeshState::New;
+
+						Impl->LatestMeshInfo.BlockInfoByGuid.Add(BlockID, BlockInfo);
+						break;
+					}
+					case MLMeshingMeshState_Deleted:
+					{
+						const auto BlockID = FGuid(MeshInfoData.id.data[0],
+												   MeshInfoData.id.data[0] >> 32,
+												   MeshInfoData.id.data[1],
+												   MeshInfoData.id.data[1] >> 32);
+
+						// Only process delete for blocks for which we actually received data
+						auto MeshBrickId = Impl->GetMrMeshIdFromMeshGuid(BlockID, false);
+						if (MeshBrickId != nullptr)
+						{
+#ifdef DEBUG_MESH_BLOCK_ADD_REMOVE
+							UE_LOG(LogMagicLeap,
+								Log,
+								TEXT("UMeshTrackerComponent: REMOVING brick %s"),
+								*(BlockID.ToString()));
+#endif //DEBUG_MESH_BLOCK_ADD_REMOVE
+							// Broadcast the mesh was removed
 							if (OnMeshTrackerUpdated.IsBound())
 							{
 								OnMeshTrackerUpdated.Broadcast(BlockID,
 									TArray<FVector>(), TArray<int32>(), TArray<FVector>(), TArray<float>());
 							}
 
-							Impl->MeshBrickCache.Remove(BlockID);
+							// Remove any pending block creations/updates
+							Impl->PendingMeshBricksByBrickId.Remove(*MeshBrickId);
+
+							if (MeshType != EMagicLeapMeshType::PointCloud)
+							{
+								// Remove brick
+								const static TArray<FVector> EmptyVertices;
+								const static TArray<FVector2D> EmptyUVs;
+								const static TArray<FPackedNormal> EmptyTangents;
+								const static TArray<FColor> EmptyVertexColors;
+								const static TArray<uint32> EmptyTriangles;
+								static_cast<IMRMesh*>(MRMesh)->SendBrickData(IMRMesh::FSendBrickDataArgs
+								{
+									nullptr,
+									*MeshBrickId,
+									EmptyVertices,
+									EmptyUVs,
+									EmptyTangents,
+									EmptyVertexColors,
+									EmptyTriangles
+								});
+							}
+						}
+						else
+						{
+#ifdef DEBUG_MESH_REQUEST_AND_RESPONSE
+							UE_LOG(LogMagicLeap,
+								Log,
+								TEXT("UMeshTrackerComponent: Received 'deleted' event for unseen block %s"),
+								*(BlockID.ToString()));
+#endif
 						}
 						break;
 					}
@@ -713,9 +904,13 @@ bool UMeshTrackerComponent::GetMeshInfoResult()
 				}
 			}
 
+			// We probably discarded some, so reduce count to the actual size
+			Impl->LatestMeshInfo.BlockData.SetNum(MeshIndex, true);
+
 			// Free up the ML meshing resources
 			MLMeshingFreeResource(Impl->MeshTracker, &Impl->CurrentMeshInfoRequest);
 			Impl->CurrentMeshInfoRequest = ML_INVALID_HANDLE;
+
 			return true;
 		}
 	}
@@ -731,12 +926,14 @@ void UMeshTrackerComponent::RequestMesh()
 	// Request block meshes for current mesh info and block list
 	if (Impl->CurrentMeshRequest == ML_INVALID_HANDLE)
 	{
-		Impl->UEMeshBlockRequests.Empty(Impl->LastestMeshInfo.BlockData.Num());
-		// Allow applications to choose which blocks to mesh.
-		IMeshBlockSelectorInterface::Execute_SelectMeshBlocks(Impl->BlockSelector.GetObject(), Impl->LastestMeshInfo, Impl->UEMeshBlockRequests);
+		Impl->UEMeshBlockRequests.Empty(Impl->LatestMeshInfo.BlockData.Num());
+
+		// Allow applications to choose which of the available blocks to mesh.
+		IMagicLeapMeshBlockSelectorInterface::Execute_SelectMeshBlocks(Impl->BlockSelector.GetObject(), Impl->LatestMeshInfo, Impl->UEMeshBlockRequests);
 
 		if (Impl->UEMeshBlockRequests.Num() > 0)
 		{
+			// Convert selected block requests to ML format
 			Impl->MeshBlockRequests.Empty(Impl->UEMeshBlockRequests.Num());
 			Impl->MeshBlockRequests.AddUninitialized(Impl->UEMeshBlockRequests.Num());
 			for (int32 i = 0; i < Impl->UEMeshBlockRequests.Num(); ++i)
@@ -744,6 +941,7 @@ void UMeshTrackerComponent::RequestMesh()
 				UnrealToMLBlockRequest(Impl->UEMeshBlockRequests[i], Impl->MeshBlockRequests[i]);
 			}
 
+			// Submit query
 			MLMeshingMeshRequest MeshRequest = {};
 			MeshRequest.request_count = static_cast<int>(Impl->MeshBlockRequests.Num());
 			MeshRequest.data = Impl->MeshBlockRequests.GetData();
@@ -786,29 +984,50 @@ bool UMeshTrackerComponent::GetMeshResult()
 		}
 		else
 		{
-			const FVector VertexOffset = UHeadMountedDisplayFunctionLibrary::GetTrackingToWorldTransform(this).Inverse().GetLocation();
+			// Translate mesh block data
+			const FTransform TrackingToWorld = UHeadMountedDisplayFunctionLibrary::GetTrackingToWorldTransform(nullptr);
 			for (uint32_t MeshIndex = 0; MeshIndex < Mesh.data_count; ++ MeshIndex)
 			{
 				const auto &MeshData = Mesh.data[MeshIndex];
 
-				auto BlockID = FGuid(MeshData.id.data[0], MeshData.id.data[0] >> 32, MeshData.id.data[1], MeshData.id.data[1] >> 32);
+				auto BlockID = FGuid(MeshData.id.data[0],
+					MeshData.id.data[0] >> 32, MeshData.id.data[1], MeshData.id.data[1] >> 32);
+
+				// Get the block info
+				auto pendingBlockInfo = Impl->LatestMeshInfo.BlockInfoByGuid.Find(BlockID);
 
-				// Create a brick index for any new mesh block
-				if (!Impl->MeshBrickCache.Contains(BlockID))
+				// Simulator can return unrequested meshes, so we cannot currently checkf
+				// https://jira.magicleap.com/browse/REM-3259
+				//checkf(pendingBlockInfo != nullptr, TEXT("Unable to find block info for pending mesh"));
+				if (pendingBlockInfo == nullptr)
 				{
-					Impl->MeshBrickCache.Add(BlockID);
+#ifdef DEBUG_MESH_REQUEST_AND_RESPONSE
+					UE_LOG(LogMagicLeap,
+						Log,
+						TEXT("UMeshTrackerComponent: Received unrequested mesh with ID %s"),
+						*(BlockID.ToString()));
+#endif
+					continue;
 				}
 
-				// Acquire mesh data cache and mark its brick ID
-				FMeshTrackerImpl::FMLCachedMeshData::SharedPtr CurrentMeshDataCache = Impl->AquireMeshDataCache();
+				// Remove them as we receive them so we can tell if we did not receive some requested blocks
+				Impl->LatestMeshInfo.BlockInfoByGuid.Remove(BlockID);
+
+				// Acquire mesh data cache and mark its ML block ID and UE brick ID
+				FMagicLeapMeshTrackerImpl::FMLCachedMeshData::SharedPtr CurrentMeshDataCache = Impl->AcquireMeshDataCache();
 				CurrentMeshDataCache->BlockID = BlockID;
+				CurrentMeshDataCache->BrickID = *Impl->GetMrMeshIdFromMeshGuid(CurrentMeshDataCache->BlockID, true);
+
+				// Copy over the block info
+				CurrentMeshDataCache->BlockInfo = *pendingBlockInfo;
 
 				// Pull vertices
 				CurrentMeshDataCache->OffsetVertices.Reserve(MeshData.vertex_count);
 				CurrentMeshDataCache->WorldVertices.Reserve(MeshData.vertex_count);
 				for (uint32_t v = 0; v < MeshData.vertex_count; ++ v)
 				{
-					CurrentMeshDataCache->OffsetVertices.Add(MagicLeap::ToFVector(MeshData.vertex[v], WorldToMetersScale) - VertexOffset);
+					const FVector WorldVertex = TrackingToWorld.TransformPosition(MagicLeap::ToFVector(MeshData.vertex[v], WorldToMetersScale));
+					CurrentMeshDataCache->OffsetVertices.Add(WorldVertex);
 					CurrentMeshDataCache->WorldVertices.Add(MagicLeap::ToFVector(MeshData.vertex[v], WorldToMetersScale));
 				}
 
@@ -863,7 +1082,7 @@ bool UMeshTrackerComponent::GetMeshResult()
 				// Apply chosen vertex color mode
 				switch (VertexColorMode)
 				{
-					case EMLMeshVertexColorMode::Confidence:
+					case EMagicLeapMeshVertexColorMode::Confidence:
 					{
 						if (nullptr != MeshData.confidence)
 						{
@@ -877,16 +1096,18 @@ bool UMeshTrackerComponent::GetMeshResult()
 						}
 						else
 						{
+							/* TODO: Replace with log once: SDKUNREAL-870
 							UE_LOG(LogMagicLeap, Warning, TEXT("MeshTracker vertex color mode is Confidence "
 								"but no confidence values available. Using white for all blocks."));
+							*/
 						}
 						break;
 					}
-					case EMLMeshVertexColorMode::Block:
+					case EMagicLeapMeshVertexColorMode::Block:
 					{
 						if (BlockVertexColors.Num() > 0)
 						{
-							const FColor& VertexColor = BlockVertexColors[(static_cast<uint64>(BlockID.B) << 32 | BlockID.A) % BlockVertexColors.Num()];
+							const FColor& VertexColor = BlockVertexColors[CurrentMeshDataCache->BrickID % BlockVertexColors.Num()];
 
 							CurrentMeshDataCache->VertexColors.Reserve(MeshData.vertex_count);
 							for (uint32 v = 0; v < MeshData.vertex_count; ++ v)
@@ -901,7 +1122,7 @@ bool UMeshTrackerComponent::GetMeshResult()
 						}
 						break;
 					}
-					case EMLMeshVertexColorMode::LOD:
+					case EMagicLeapMeshVertexColorMode::LOD:
 					{
 						if (BlockVertexColors.Num() >= MLMeshingLOD_Maximum)
 						{
@@ -920,7 +1141,7 @@ bool UMeshTrackerComponent::GetMeshResult()
 						}
 						break;
 					}
-					case EMLMeshVertexColorMode::None:
+					case EMagicLeapMeshVertexColorMode::None:
 					{
 						break;
 					}
@@ -946,33 +1167,69 @@ bool UMeshTrackerComponent::GetMeshResult()
 					CurrentMeshDataCache->UV0.Add(FVector2D(FakeCoord, FakeCoord));
 				}
 
-				// Create/update brick
-				if (MeshType != EMeshType::PointCloud)
+				if (BricksPerFrame > 0)
 				{
-					static_cast<IMRMesh*>(MRMesh)->SendBrickData(IMRMesh::FSendBrickDataArgs
-						{
-							TSharedPtr<IMRMesh::FBrickDataReceipt, ESPMode::ThreadSafe>
-								(new FMeshTrackerImpl::FMeshTrackerComponentBrickDataReceipt(CurrentMeshDataCache)),
-							(static_cast<uint64>(CurrentMeshDataCache->BlockID.B) << 32 | CurrentMeshDataCache->BlockID.A),
-							CurrentMeshDataCache->WorldVertices,
-							CurrentMeshDataCache->UV0,
-							CurrentMeshDataCache->Tangents,
-							CurrentMeshDataCache->VertexColors,
-							CurrentMeshDataCache->Triangles
-						}
-					);
+					Impl->PendingMeshBricksByBrickId.Add(CurrentMeshDataCache->BrickID, CurrentMeshDataCache);
 				}
+				else
+				{
+#ifdef DEBUG_MESH_BLOCK_ADD_REMOVE
+					UE_LOG(LogMagicLeap,
+						Log,
+						TEXT("UMeshTrackerComponent: ADDING/UPDATING brick %s"),
+						*(BlockID.ToString()));
+#endif //DEBUG_MESH_BLOCK_ADD_REMOVE
+					// Broadcast that a mesh was updated
+					if (OnMeshTrackerUpdated.IsBound())
+					{
+						// Hack because blueprints don't support uint32.
+						TArray<int32> Triangles(reinterpret_cast<const int32*>(CurrentMeshDataCache->
+							Triangles.GetData()), CurrentMeshDataCache->Triangles.Num());
+						OnMeshTrackerUpdated.Broadcast(CurrentMeshDataCache->BlockID, CurrentMeshDataCache->OffsetVertices,
+							Triangles, CurrentMeshDataCache->Normals, CurrentMeshDataCache->Confidence);
+					}
 
-				// Broadcast that a mesh was updated
-				if (OnMeshTrackerUpdated.IsBound())
+					if (MeshType != EMagicLeapMeshType::PointCloud)
+					{
+						// Create/update brick
+						static_cast<IMRMesh*>(MRMesh)->SendBrickData(IMRMesh::FSendBrickDataArgs
+							{
+								TSharedPtr<IMRMesh::FBrickDataReceipt, ESPMode::ThreadSafe>
+									(new FMagicLeapMeshTrackerImpl::FMeshTrackerComponentBrickDataReceipt(CurrentMeshDataCache)),
+									CurrentMeshDataCache->BrickID,
+									CurrentMeshDataCache->WorldVertices,
+									CurrentMeshDataCache->UV0,
+									CurrentMeshDataCache->Tangents,
+									CurrentMeshDataCache->VertexColors,
+									CurrentMeshDataCache->Triangles,
+									FBox(CurrentMeshDataCache->BlockInfo.BlockPosition - CurrentMeshDataCache->BlockInfo.BlockDimensions,
+										 CurrentMeshDataCache->BlockInfo.BlockPosition + CurrentMeshDataCache->BlockInfo.BlockDimensions)
+							});
+					}
+					else
+					{
+						// Discard
+						CurrentMeshDataCache->Recycle(CurrentMeshDataCache);
+					}
+				}
+			}
+
+#ifdef DEBUG_MESH_REQUEST_AND_RESPONSE
+			if (Impl->LatestMeshInfo.BlockInfoByGuid.Num() != 0)
+			{
+				UE_LOG(LogMagicLeap,
+					Log,
+					TEXT("UMeshTrackerComponent: Failed to receive meshes for %d requests:"),
+					Impl->LatestMeshInfo.BlockInfoByGuid.Num());
+				for (const auto& BlockInfo : Impl->LatestMeshInfo.BlockInfoByGuid)
 				{
-					// Hack because blueprints don't support uint32.
-					TArray<int32> Triangles(reinterpret_cast<const int32*>(CurrentMeshDataCache->
-						Triangles.GetData()), CurrentMeshDataCache->Triangles.Num());
-					OnMeshTrackerUpdated.Broadcast(CurrentMeshDataCache->BlockID, CurrentMeshDataCache->OffsetVertices,
-						Triangles, CurrentMeshDataCache->Normals, CurrentMeshDataCache->Confidence);
+					UE_LOG(LogMagicLeap,
+						Log,
+						TEXT("UMeshTrackerComponent: No data received for %s"),
+						*BlockInfo.Key.ToString());
 				}
 			}
+#endif
 
 			// All meshes pulled and/or updated; free the ML resource
 			MLMeshingFreeResource(Impl->MeshTracker, &Impl->CurrentMeshRequest);
diff --git a/Engine/Plugins/Lumin/MagicLeap/Source/MagicLeap/Private/RaycastComponent.cpp b/Engine/Plugins/Lumin/MagicLeap/Source/MagicLeap/Private/RaycastComponent.cpp
--- a/Engine/Plugins/Lumin/MagicLeap/Source/MagicLeap/Private/RaycastComponent.cpp
+++ b/Engine/Plugins/Lumin/MagicLeap/Source/MagicLeap/Private/RaycastComponent.cpp
@@ -9,18 +9,12 @@
 #if WITH_EDITOR
 #include "Editor.h"
 #endif
-#include "MagicLeapPluginUtil.h" // for ML_INCLUDES_START/END
+#include "Lumin/CAPIShims/LuminAPIRaycast.h"
 
-#if WITH_MLSDK
-ML_INCLUDES_START
-#include <ml_raycast.h>
-ML_INCLUDES_END
-#endif //WITH_MLSDK
-
-class FRaycastTrackerImpl
+class FMagicLeapRaycastTrackerImpl
 {
 public:
-	FRaycastTrackerImpl()
+	FMagicLeapRaycastTrackerImpl()
 #if WITH_MLSDK
 		: Tracker(ML_INVALID_HANDLE)
 #endif //WITH_MLSDK
@@ -40,7 +34,7 @@ public:
 			MLResult Result = MLRaycastCreate(&Tracker);
 			if (Result != MLResult_Ok)
 			{
-				UE_LOG(LogMagicLeap, Error, TEXT("MLRaycastCreate failed with error %d."), Result);
+				UE_LOG(LogMagicLeap, Error, TEXT("MLRaycastCreate failed with error %s."), UTF8_TO_TCHAR(MLGetResultString(Result)));
 				return false;
 			}
 		}
@@ -54,7 +48,7 @@ public:
 		if (MLHandleIsValid(Tracker))
 		{
 			MLResult Result = MLRaycastDestroy(Tracker);
-			UE_CLOG(Result != MLResult_Ok, LogMagicLeap, Error, TEXT("MLRaycastDestroy failed with error %d."), Result);
+			UE_CLOG(Result != MLResult_Ok, LogMagicLeap, Error, TEXT("MLRaycastDestroy failed with error %s."), UTF8_TO_TCHAR(MLGetResultString(Result)));
 			Tracker = ML_INVALID_HANDLE;
 		}
 #endif //WITH_MLSDK
@@ -62,46 +56,46 @@ public:
 };
 
 #if WITH_MLSDK
-ERaycastResultState MLToUnrealRaycastResultState(MLRaycastResultState state)
+EMagicLeapRaycastResultState MLToUnrealRaycastResultState(MLRaycastResultState state)
 {
 	switch (state)
 	{
 	case MLRaycastResultState_RequestFailed:
-		return ERaycastResultState::RequestFailed;
+		return EMagicLeapRaycastResultState::RequestFailed;
 	case MLRaycastResultState_HitObserved:
-		return ERaycastResultState::HitObserved;
+		return EMagicLeapRaycastResultState::HitObserved;
 	case MLRaycastResultState_HitUnobserved:
-		return ERaycastResultState::HitUnobserved;
+		return EMagicLeapRaycastResultState::HitUnobserved;
 	case MLRaycastResultState_NoCollision:
-		return ERaycastResultState::NoCollision;
+		return EMagicLeapRaycastResultState::NoCollision;
 	}
-	return ERaycastResultState::RequestFailed;
+	return EMagicLeapRaycastResultState::RequestFailed;
 }
 #endif //WITH_MLSDK
 
-URaycastComponent::URaycastComponent()
-	: Impl(new FRaycastTrackerImpl())
+UMagicLeapRaycastComponent::UMagicLeapRaycastComponent()
+	: Impl(new FMagicLeapRaycastTrackerImpl())
 {
 	// Make sure this component ticks
 	PrimaryComponentTick.bCanEverTick = true;
 	PrimaryComponentTick.bStartWithTickEnabled = true;
-	PrimaryComponentTick.TickGroup = TG_PrePhysics;
+	PrimaryComponentTick.TickGroup = TG_PostPhysics;
 	bAutoActivate = true;
 
 #if WITH_EDITOR
 	if (GIsEditor)
 	{
-		FEditorDelegates::PrePIEEnded.AddUObject(this, &URaycastComponent::PrePIEEnded);
+		FEditorDelegates::PrePIEEnded.AddUObject(this, &UMagicLeapRaycastComponent::PrePIEEnded);
 	}
 #endif
 }
 
-URaycastComponent::~URaycastComponent()
+UMagicLeapRaycastComponent::~UMagicLeapRaycastComponent()
 {
 	delete Impl;
 }
 
-void URaycastComponent::TickComponent(float DeltaTime, enum ELevelTick TickType, FActorComponentTickFunction *ThisTickFunction)
+void UMagicLeapRaycastComponent::TickComponent(float DeltaTime, enum ELevelTick TickType, FActorComponentTickFunction *ThisTickFunction)
 {
 	Super::TickComponent(DeltaTime, TickType, ThisTickFunction);
 
@@ -118,13 +112,13 @@ void URaycastComponent::TickComponent(float DeltaTime, enum ELevelTick TickType,
 		if (APICallResult == MLResult_Ok)
 		{
 			const FAppFramework& AppFramework = static_cast<FMagicLeapHMD*>(GEngine->XRSystem->GetHMDDevice())->GetAppFrameworkConst();
-			float WorldToMetersScale = AppFramework.IsInitialized() ? AppFramework.GetWorldToMetersScale() : 100.0f;
+			const float WorldToMetersScale = AppFramework.GetWorldToMetersScale();
 
 			// TODO: Should we apply this transform here or expect the user to use the result as a child of the XRPawn like the other features?
 			// This being for raycast, we should probably apply the transform since the result might be used for other than just placing objects.
 			const FTransform TrackingToWorld = UHeadMountedDisplayFunctionLibrary::GetTrackingToWorldTransform(this);
 
-			FRaycastHitResult hitResult;
+			FMagicLeapRaycastHitResult hitResult;
 			hitResult.HitState = MLToUnrealRaycastResultState(result.state);
 			hitResult.HitPoint = TrackingToWorld.TransformPosition(MagicLeap::ToFVector(result.hitpoint, WorldToMetersScale));
 			hitResult.Normal = TrackingToWorld.TransformVectorNoScale(MagicLeap::ToFVector(result.normal, 1.0f));
@@ -134,7 +128,7 @@ void URaycastComponent::TickComponent(float DeltaTime, enum ELevelTick TickType,
 			if (hitResult.HitPoint.ContainsNaN() || hitResult.Normal.ContainsNaN())
 			{
 				UE_LOG(LogMagicLeap, Error, TEXT("Raycast result contains NaNs."));
-				hitResult.HitState = ERaycastResultState::RequestFailed;
+				hitResult.HitState = EMagicLeapRaycastResultState::RequestFailed;
 			}
 
 			pair.Value.ResultDelegate.ExecuteIfBound(hitResult);
@@ -142,7 +136,7 @@ void URaycastComponent::TickComponent(float DeltaTime, enum ELevelTick TickType,
 		}
 		else if (APICallResult != MLResult_Pending)
 		{
-			UE_LOG(LogMagicLeap, Error, TEXT("MLRaycastGetResult failed with result %d."), APICallResult);
+			UE_LOG(LogMagicLeap, Error, TEXT("MLRaycastGetResult failed with result %s."), UTF8_TO_TCHAR(MLGetResultString(APICallResult)));
 		}
 	}
 
@@ -158,7 +152,7 @@ void URaycastComponent::TickComponent(float DeltaTime, enum ELevelTick TickType,
 #endif //WITH_MLSDK
 }
 
-bool URaycastComponent::RequestRaycast(const FRaycastQueryParams& RequestParams, const FRaycastResultDelegate& ResultDelegate)
+bool UMagicLeapRaycastComponent::RequestRaycast(const FMagicLeapRaycastQueryParams& RequestParams, const FRaycastResultDelegate& ResultDelegate)
 {
 #if WITH_MLSDK
 	if (!(IMagicLeapPlugin::Get().IsMagicLeapHMDValid() && Impl->Create()))
@@ -167,7 +161,7 @@ bool URaycastComponent::RequestRaycast(const FRaycastQueryParams& RequestParams,
 	}
 
 	const FAppFramework& AppFramework = static_cast<FMagicLeapHMD*>(GEngine->XRSystem->GetHMDDevice())->GetAppFrameworkConst();
-	float WorldToMetersScale = AppFramework.IsInitialized() ? AppFramework.GetWorldToMetersScale() : 100.0f;
+	const float WorldToMetersScale = AppFramework.GetWorldToMetersScale();
 
 	const FTransform WorldToTracking = UHeadMountedDisplayFunctionLibrary::GetTrackingToWorldTransform(this).Inverse();
 
@@ -184,7 +178,7 @@ bool URaycastComponent::RequestRaycast(const FRaycastQueryParams& RequestParams,
 	MLResult Result = MLRaycastRequest(Impl->Tracker, &query, &Handle);
 	if (Result != MLResult_Ok)
 	{
-		UE_LOG(LogMagicLeap, Error, TEXT("MLRaycastRequest failed with error %d."), Result);
+		UE_LOG(LogMagicLeap, Error, TEXT("MLRaycastRequest failed with error %s."), UTF8_TO_TCHAR(MLGetResultString(Result)));
 		return false;
 	}
 
@@ -196,7 +190,7 @@ bool URaycastComponent::RequestRaycast(const FRaycastQueryParams& RequestParams,
 	return true;
 }
 
-void URaycastComponent::FinishDestroy()
+void UMagicLeapRaycastComponent::FinishDestroy()
 {
 #if WITH_EDITOR
 	if (GIsEditor)
@@ -209,8 +203,22 @@ void URaycastComponent::FinishDestroy()
 }
 
 #if WITH_EDITOR
-void URaycastComponent::PrePIEEnded(bool bWasSimulatingInEditor)
+void UMagicLeapRaycastComponent::PrePIEEnded(bool bWasSimulatingInEditor)
 {
 	Impl->Destroy();
 }
 #endif
+
+FMagicLeapRaycastQueryParams UMagicLeapRaycastFunctionLibrary::MakeRaycastQueryParams(FVector Position, FVector Direction, FVector UpVector, int32 Width, int32 Height, float HorizontalFovDegrees, bool CollideWithUnobserved, int32 UserData)
+{
+	FMagicLeapRaycastQueryParams QueryParams;
+	QueryParams.Position = Position;
+	QueryParams.Direction = Direction;
+	QueryParams.UpVector = UpVector;
+	QueryParams.Width = Width;
+	QueryParams.Height = Height;
+	QueryParams.HorizontalFovDegrees = HorizontalFovDegrees;
+	QueryParams.CollideWithUnobserved = CollideWithUnobserved;
+	QueryParams.UserData = UserData;
+	return QueryParams;
+}
diff --git a/Engine/Plugins/Lumin/MagicLeap/Source/MagicLeap/Public/AppEventHandler.h b/Engine/Plugins/Lumin/MagicLeap/Source/MagicLeap/Public/AppEventHandler.h
new file mode 100644
--- /dev/null
+++ b/Engine/Plugins/Lumin/MagicLeap/Source/MagicLeap/Public/AppEventHandler.h
@@ -0,0 +1,168 @@
+// Copyright 1998-2019 Epic Games, Inc. All Rights Reserved.
+
+#pragma once
+
+#include "CoreMinimal.h"
+#include "MagicLeapPrivilegeTypes.h"
+#include "Lumin/CAPIShims/LuminAPI.h"
+
+namespace MagicLeap
+{
+	enum EPrivilegeState
+	{
+		NotYetRequested,
+		Pending,
+		Granted,
+		Denied,
+		Error
+	};
+
+	struct FRequiredPrivilege
+	{
+		FRequiredPrivilege(EMagicLeapPrivilege InPrivilegeID)
+		: PrivilegeID(InPrivilegeID)
+		, PrivilegeRequest(nullptr)
+		, State(NotYetRequested)
+		{}
+
+		EMagicLeapPrivilege PrivilegeID;
+		void* PrivilegeRequest;
+		EPrivilegeState State;
+	};
+
+	/**
+	* Provides an interface between the AppFramework and any system that needs to be
+	* notified of application events (such as pause/resume).
+	*/
+	class MAGICLEAP_API IAppEventHandler
+	{
+	public:
+		typedef TFunction<void()> FEventHandler;
+		typedef TFunction<void(const FRequiredPrivilege& RequiredPrivilege)> FPrivilegeEventHandler;
+
+		/** 
+			Adds the IAppEventHandler instance to the application's list of IAppEventHandler instances.
+			Populates a RequiredPrivileges list based on the privilge ids passed via InRequiredPrivileges;
+			@param InRequiredPrivileges The list of privilge ids required by the calling system.
+		*/
+		IAppEventHandler(const TArray<EMagicLeapPrivilege>& InRequiredPrivileges);
+
+		IAppEventHandler();
+
+		/** Removes the IAppEventHandler instance from the application's list of IAppEventHandler instances.*/
+		virtual ~IAppEventHandler();
+
+		/**
+			Can be overridden by inheriting class that needs to destroy certain api interfaces before the perception stack is
+			closed down.
+		*/
+		virtual void OnAppShutDown();
+
+		/**
+			Use to check status of privilege requests.
+		*/
+		virtual void OnAppTick();
+
+		/**
+			Can be overridden by inheriting class in order to pause its system.
+		*/
+		virtual void OnAppPause();
+
+		/**
+			Can be overridden by inheriting class in order to resume its system.
+		*/
+		virtual void OnAppResume();
+
+		/**
+			Returns the status of the specified privilege.
+			@param PrivilegeID The privilege id to be queried.
+			@param bBlocking Flags whether or not to use the blocking query internally.
+		*/
+		EPrivilegeState GetPrivilegeStatus(EMagicLeapPrivilege PrivilegeID, bool bBlocking = true);
+
+		/**
+			Converts the EMagicLeapPrivilege enum value to it's corresponding string representation.
+			@param PrivilegeID The EMagicLeapPrivilege enum value to be converted to a string.
+			@return An FString representing the string value of the supplied privilege enum value.
+		*/
+		FString PrivilegeToString(EMagicLeapPrivilege PrivilegeID);
+
+		/**
+			Converts the EPrivilegeState enum value to it's corresponding string representation.
+			@param PrivilegeState The EPrivilegeState enum value to be converted to a string.
+			@return A TCHAR pointer representing the string value of the supplied privilege state enum value.
+		*/
+		const TCHAR* PrivilegeStateToString(EPrivilegeState PrivilegeState);
+
+		/**
+			Pushes this object onto a worker thread so that it's blocking destructor can be called without locking up the update thread.
+			@note This should only be called by objects that have a blocking destructor and are no longer referenced.
+		*/
+		bool AsyncDestroy();
+
+		/**
+			Triggered when a privilege request changes state.
+		*/
+		void SetPrivilegeEventHandler(FPrivilegeEventHandler&& InOnPrivilegeEvent)
+		{
+			OnPrivilegeEvent = MoveTemp(InOnPrivilegeEvent);
+		}
+
+		/**
+			Use this as an alternative to overriding the OnAppShutDown function.  This allows you to use IAppEventHandler
+			as and aggregate class rather than an ancestor. 
+		*/
+		void SetOnAppShutDownHandler(FEventHandler&& InOnAppShutDownHandler)
+		{
+			OnAppShutDownHandler = MoveTemp(InOnAppShutDownHandler);
+		}
+
+		/**
+			Use this as an alternative to overriding the OnAppTick function.  This allows you to use IAppEventHandler
+			as and aggregate class rather than an ancestor.
+		*/
+		void SetOnAppTickHandler(FEventHandler&& InOnAppTickHandler)
+		{
+			OnAppTickHandler = MoveTemp(InOnAppTickHandler);
+		}
+
+		/**
+			Use this as an alternative to overriding the OnAppPause function.  This allows you to use IAppEventHandler
+			as and aggregate class rather than an ancestor.
+		*/
+		void SetOnAppPauseHandler(FEventHandler&& InOnAppPauseHandler)
+		{
+			OnAppPauseHandler = MoveTemp(InOnAppPauseHandler);
+		}
+
+		/**
+			Use this as an alternative to overriding the OnAppResume function.  This allows you to use IAppEventHandler
+			as and aggregate class rather than an ancestor.
+		*/
+		void SetOnAppResumeHandler(FEventHandler&& InOnAppResumeHandler)
+		{
+			OnAppResumeHandler = MoveTemp(InOnAppResumeHandler);
+		}
+
+		void SetWasSystemEnabledOnPause(bool bInWasSystemEnabledOnPause)
+		{
+			bWasSystemEnabledOnPause = bInWasSystemEnabledOnPause;
+		}
+
+		bool WasSystemEnabledOnPause() const
+		{
+			return bWasSystemEnabledOnPause;
+		}
+
+	protected:
+		TMap<EMagicLeapPrivilege, FRequiredPrivilege> RequiredPrivileges;
+		FPrivilegeEventHandler OnPrivilegeEvent;
+		FEventHandler OnAppShutDownHandler;
+		FEventHandler OnAppTickHandler;
+		FEventHandler OnAppPauseHandler;
+		FEventHandler OnAppResumeHandler;
+		bool bAllPrivilegesInSync;
+		bool bWasSystemEnabledOnPause;
+		FCriticalSection CriticalSection;
+	};
+} // MagicLeap
diff --git a/Engine/Plugins/Lumin/MagicLeap/Source/MagicLeap/Public/IMagicLeapPlugin.h b/Engine/Plugins/Lumin/MagicLeap/Source/MagicLeap/Public/IMagicLeapPlugin.h
--- a/Engine/Plugins/Lumin/MagicLeap/Source/MagicLeap/Public/IMagicLeapPlugin.h
+++ b/Engine/Plugins/Lumin/MagicLeap/Source/MagicLeap/Public/IMagicLeapPlugin.h
@@ -5,9 +5,22 @@
 #include "CoreMinimal.h"
 #include "Modules/ModuleManager.h"
 #include "IHeadMountedDisplayModule.h"
-#include "IMagicLeapHMD.h"
+#include "Stats/Stats.h"
+#include "Lumin/CAPIShims/LuminAPICoordinateFrameUID.h"
 
-class IMagicLeapInputDevice;
+class IMagicLeapTrackerEntity;
+
+DECLARE_STATS_GROUP(TEXT("MagicLeap"), STATGROUP_MagicLeap, STATCAT_Advanced);
+
+enum class EMagicLeapTransformFailReason : uint8
+{
+	None,
+	InvalidTrackingFrame,
+	NaNsInTransform,
+	CallFailed,
+	PoseNotFound,
+	HMDNotInitialized
+};
 
 /**
  * The public interface to this module.  In most cases, this interface is only public to sibling modules
@@ -28,6 +41,17 @@ public:
 		return FModuleManager::LoadModuleChecked< IMagicLeapPlugin >("MagicLeap");
 	}
 
+	/**
+	 * Singleton-like access to this module's interface.
+	 * Checks whether the module is loaded, asserting on failure.
+	 *
+	 * @return Returns singleton instance
+	 */
+	static inline IMagicLeapPlugin& GetModuleChecked()
+	{
+		return FModuleManager::GetModuleChecked< IMagicLeapPlugin >("MagicLeap");
+	}
+
 	/**
 	 * Checks to see if this module is loaded and ready.  It is only valid to call Get() if IsAvailable() returns true.
 	 *
@@ -46,15 +70,76 @@ public:
 	virtual bool IsMagicLeapHMDValid() const = 0;
 
 	/**
-	* Returns the HMD associated with the ML plugin.
-	*
-	* @return Shared pointer to the HMD.
-	*/
-	virtual TWeakPtr<IMagicLeapHMD, ESPMode::ThreadSafe> GetHMD() = 0;
+	 * Register an IMagicLeapTrackerEntity with the HMD to receive callbacks for tracker creation, desctruction and late update.
+	 * @param TrackerEntity Pointer to the object that implements the IMagicLeapTrackerEntity interface.
+	 */
+	virtual void RegisterMagicLeapTrackerEntity(IMagicLeapTrackerEntity* TrackerEntity) = 0;
+
+	/**
+	 * Unregister an IMagicLeapTrackerEntity with the HMD. Object will no longer receive callbacks for tracker creation, desctruction and late update.
+	 * @param TrackerEntity Pointer to the object that implements the IMagicLeapTrackerEntity interface.
+	 */
+	virtual void UnregisterMagicLeapTrackerEntity(IMagicLeapTrackerEntity* TrackerEntity) = 0;
+
+	/**
+	 * Enable/Disable mouse and keyboard input when playing in Editor.
+	 * @param Ignore If true, mouse and keyboard input is ignored by the viewport. If false, the input enabled again.
+	 * @return Whether input is being ignored now or not.
+	 */
+	virtual bool SetIgnoreInput(bool Ignore) = 0;
+
+	/**
+	 * Get latest transform for a given CoordinateFrameUID in Tracking space. Transform to World space using UHeadMountedDisplayFunctionLibrary::GetTrackingToWorldTransform().
+	 * @param Id The CoordinateFrameUID, as an FGuid, to get the transform for.
+	 * @param OutTransform Transform of the given CoordinateFrameUID.
+	 * @param OutReason Reason for this function's failure.
+	 * @return true if OutTransform is populated with a valid value. If false, the reason for failure is populated in OutReason.
+	 */
+	virtual bool GetTransform(const FGuid& Id, FTransform& OutTransform, EMagicLeapTransformFailReason& OutReason) const = 0;
+
+#if WITH_MLSDK
+	/**
+	 * Get latest transform for a given CoordinateFrameUID in Tracking space. Transform to World space using UHeadMountedDisplayFunctionLibrary::GetTrackingToWorldTransform().
+	 * @param Id The CoordinateFrameUID to get the transform for.
+	 * @param OutTransform Transform of the given CoordinateFrameUID.
+	 * @param OutReason Reason for this function's failure.
+	 * @return true if OutTransform is populated with a valid value. If false, the reason for failure is populated in OutReason.
+	 */
+	virtual bool GetTransform(const MLCoordinateFrameUID& Id, FTransform& OutTransform, EMagicLeapTransformFailReason& OutReason) const = 0;
+#endif // WITH_MLSDK
+
+	/**
+	 * Scale of 1uu to 1m in real world measurements.
+	 * @return Default value is 100.0f i.e. 1uu = 100cm.
+	 */
+	virtual float GetWorldToMetersScale() const = 0;
+
+	/**
+	 * Check if the perception system has been started.
+	 * @return true if perception system has been started, false otherwise.
+	 */
+	virtual bool IsPerceptionEnabled() const = 0;
 
-	virtual void RegisterMagicLeapInputDevice(IMagicLeapInputDevice* InputDevice) = 0;
-	virtual void UnregisterMagicLeapInputDevice(IMagicLeapInputDevice* InputDevice) = 0;
-	virtual void EnableInputDevices() = 0;
-	virtual void DisableInputDevices() = 0;
-	virtual void OnBeginRendering_GameThread_UpdateInputDevices() = 0;
+	/**
+	 * Check if MLAudio mixer should be used when running with MLRemote.
+	 * @return true MLAudio mixer should be used, false otherwise.
+	 */
+	virtual bool UseMLAudioForZI() const = 0;
+};
+
+class IARSystemSupport;
+class FXRTrackingSystemBase;
+
+/**
+ * The public interface to this module.
+ */
+class ILuminARModule : public IModuleInterface
+{
+public:
+	//create for mutual connection (regardless of construction order)
+	virtual TSharedPtr<IARSystemSupport, ESPMode::ThreadSafe> CreateARImplementation() = 0;
+	//Now connect (regardless of connection order)
+	virtual void ConnectARImplementationToXRSystem(FXRTrackingSystemBase* InXRTrackingSystem) = 0;
+	//Now initialize fully connected systems
+	virtual void InitializeARImplementation() = 0;
 };
diff --git a/Engine/Plugins/Lumin/MagicLeap/Source/MagicLeap/Public/IMagicLeapTrackerEntity.h b/Engine/Plugins/Lumin/MagicLeap/Source/MagicLeap/Public/IMagicLeapTrackerEntity.h
new file mode 100644
--- /dev/null
+++ b/Engine/Plugins/Lumin/MagicLeap/Source/MagicLeap/Public/IMagicLeapTrackerEntity.h
@@ -0,0 +1,48 @@
+// Copyright 1998-2019 Epic Games, Inc. All Rights Reserved.
+
+#pragma once
+
+#include "IMagicLeapPlugin.h"
+
+/**
+	The public interface for MagicLeapHMD to control the order of tracker creation, desctruction, provide late update etc.
+
+	Perception related trackers should only be created after a successful call to MLPerceptionStartup() and be destroyed
+	before a call to MLPerceptionShutdown(). This interface will ensure that.
+
+	Entities must register with the IMagicLeapPlugin using it's RegisterMagicLeapTrackerEntity() function.
+*/
+class IMagicLeapTrackerEntity
+{
+public:
+	virtual void CreateEntityTracker() {}
+	virtual void DestroyEntityTracker() {}
+	virtual void OnBeginRendering_GameThread() {}
+};
+
+namespace MagicLeap
+{
+	/** 
+	 * Utility class to scope guard enabling and disabling game viewport client input processing.
+	 * On creation it will enable the input processing, and on exit it will restore it to its
+	 * previous state. Usage is:
+	 *
+	 * { MagicLeap::EnableInput EnableInput; PostSomeInputToMessageHandlers(); }
+	 */
+	struct EnableInput
+	{
+	#if WITH_EDITOR
+		inline EnableInput()
+		{
+			SavedIgnoreInput = IMagicLeapPlugin::Get().SetIgnoreInput(false);
+		}
+		inline ~EnableInput()
+		{
+			IMagicLeapPlugin::Get().SetIgnoreInput(SavedIgnoreInput);
+		}
+
+	private:
+		bool SavedIgnoreInput;
+	#endif
+	};
+}
diff --git a/Engine/Plugins/Lumin/MagicLeap/Source/MagicLeap/Public/InAppPurchaseComponent.h b/Engine/Plugins/Lumin/MagicLeap/Source/MagicLeap/Public/InAppPurchaseComponent.h
new file mode 100644
--- /dev/null
+++ b/Engine/Plugins/Lumin/MagicLeap/Source/MagicLeap/Public/InAppPurchaseComponent.h
@@ -0,0 +1,220 @@
+// Copyright 1998-2019 Epic Games, Inc. All Rights Reserved.
+
+#pragma once
+
+#include "Components/ActorComponent.h"
+#include "Engine/Engine.h"
+#include "InAppPurchaseComponent.generated.h"
+
+UENUM(BlueprintType)
+enum class PurchaseType : uint8
+{
+	/*!
+	\brief This represents an item that can be bought multiple times. The
+	application is responsible for managing the consumption of this item.
+	*/
+	Consumable,
+	/*!
+	\brief This represents an item that can only be bought once. This will
+	be enforced by the services.
+	*/
+	Nonconsumable,
+
+	Undefined,
+};
+
+UENUM(BlueprintType)
+enum class CloudStatus : uint8
+{
+	/*! The current request is still in-progress. */
+	CloudStatus_NotDone = 0,
+	/*! The current request is complete. It may have succeeded or failed. */
+	CloudStatus_Done,
+};
+
+USTRUCT(BlueprintType)
+struct FPurchaseItemDetails
+{
+	GENERATED_BODY()
+
+	/*! This is the id of the item. */
+	FString IAPID;
+	/*! This is the formatted price for the item. */
+	UPROPERTY(EditAnywhere, BlueprintReadOnly, Category = "PurchaseItemDetails|MagicLeap")
+	FString Price;
+	/*! This is the name of the item. */
+	UPROPERTY(EditAnywhere, BlueprintReadOnly, Category = "PurchaseItemDetails|MagicLeap")
+	FString Name;
+	/*! This is the type of purchase. */
+	UPROPERTY(EditAnywhere, BlueprintReadOnly, Category = "PurchaseItemDetails|MagicLeap")
+	PurchaseType Type;
+	/*! This is the token to be used when submitting a purchase. */
+	const char *Token;
+};
+
+USTRUCT(BlueprintType)
+struct FPurchaseConfirmation
+{
+	GENERATED_BODY()
+	/*! This is the unique order id for this purchase. */
+	FString OrderID;
+	/*! This is the name of the item from where this purchase originated. */
+	UPROPERTY(EditAnywhere, BlueprintReadOnly, Category = "PurchaseConfirmation|MagicLeap")
+	FString PackageName;
+	/*! This is the time the product was purchased, in milliseconds since the epoch (Jan 1, 1970). */
+	uint64 PurchaseTime;
+	/*!
+	\brief This is a string containing the signature of the purchase data that
+	was signed with the private key of the developer.
+	*/
+	FString Signature;
+	/*! This is the in app purchase ID of the item being purchased. */
+	FString IAPID;
+	/*! This is the type of purchase. */
+	UPROPERTY(EditAnywhere, BlueprintReadOnly, Category = "PurchaseConfirmation|MagicLeap")
+	PurchaseType Type;
+};
+
+/**
+	The CameraCaptureComponent provides access to and maintains state for camera capture functionality.
+	The connection to the device's camera is managed internally.  Users of this component
+	are able to asynchronously capture camera images and footage to file.  Alternatively,
+	a camera image can be captured directly to texture.  The user need only make the relevant
+	asynchronous call and then register the appropriate success/fail event handlers for the
+	operation's completion.
+*/
+UCLASS(ClassGroup = MagicLeap, BlueprintType, Blueprintable, EditInlineNew, meta = (BlueprintSpawnableComponent))
+class MAGICLEAP_API UInAppPurchaseComponent : public UActorComponent
+{
+	GENERATED_BODY()
+
+public:
+	UInAppPurchaseComponent();
+	virtual ~UInAppPurchaseComponent();
+
+	/** Intializes internal systems. */
+	void BeginPlay() override;
+
+	/** Polls for query results. */
+	virtual void TickComponent(float DeltaTime, enum ELevelTick TickType, FActorComponentTickFunction *ThisTickFunction) override;
+
+	/**
+		Attempts to retrieve details for the specified items.
+		@brief	This call instigates an items details query which is handled on a separate thread.  The result of this asynchronous
+				operation is reported back to the calling blueprint via the FGetItemsDetailsSuccess or FGetItemsDetailsFailure event handlers.
+		@param	ItemIDs An array of FString objects specifying the names of the items whose details will be queried.
+		@return	False if an items details query is already running, true otherwise.
+	*/
+	UFUNCTION(BlueprintCallable, Category = "InAppPurchase|MagicLeap")
+	bool TryGetItemsDetailsAsync(const TArray<FString>& ItemIDs);
+
+	/**
+		Attempts to purchase the specified item.
+		@brief	This call instigates a purchase request which is handled on a separate thread.  The result of this asynchronous operation is
+				reported back to the calling blueprint via the FPurchaseConfirmationSuccess or FPurchaseConfirmationFailure event handlers.
+		@param	ItemDetails The details of the item to be purchased.
+		@return	False if a purchase confirmation is already running, true otherwise.
+	*/
+	UFUNCTION(BlueprintCallable, Category = "InAppPurchase|MagicLeap")
+	bool TryPurchaseItemAsync(const FPurchaseItemDetails& ItemDetails);
+
+	/**
+		Attempts to retrieve the app's purchase history.
+		@brief	This call instigates a purchase history request which is handled on a separate thread.  The result of this asynchronous
+				operation is reported back to the calling blueprint via the FGetPurchaseHistorySuccess or FGetPurchaseHistoryFailure event handlers.
+		@param	InNumPages Specifies the number of history pages to retrieve.
+		@return	False if a purchase history query is already running or InNumPages is less than or equal to zero, true otherwise.
+	*/
+	UFUNCTION(BlueprintCallable, Category = "InAppPurchase|MagicLeap")
+	bool TryGetPurchaseHistoryAsync(int32 InNumPages);
+
+public:
+	/**
+		Delegate used to report log messages.
+		@note This is useful if the user wishes to have log messages in 3D space.
+		@param LogMessage A string containing the log message.
+	*/
+	DECLARE_DYNAMIC_MULTICAST_DELEGATE_OneParam(FInAppPurchaseLogMessage, const FString&, LogMessage);
+
+	/**
+		Delegate used to report a successful retrieval of items details.
+		@param ItemsDetails A list of items details.
+	*/
+	DECLARE_DYNAMIC_MULTICAST_DELEGATE_OneParam(FGetItemsDetailsSuccess, const TArray<FPurchaseItemDetails>&, ItemsDetails);
+
+	/**
+		Delegate used to report a failure to retrieve the requested items details.
+	*/
+	DECLARE_DYNAMIC_MULTICAST_DELEGATE(FGetItemsDetailsFailure);
+
+	/**
+		Delegate used to report a successful item purchase confirmation.
+		@param PurchaseConfirmations A list of item purchase confirmations.
+	*/
+	DECLARE_DYNAMIC_MULTICAST_DELEGATE_OneParam(FPurchaseConfirmationSuccess, const FPurchaseConfirmation&, PurchaseConfirmations);
+
+	/**
+		Delegate used to report a failure to retrieve an item purchase confirmation.
+	*/
+	DECLARE_DYNAMIC_MULTICAST_DELEGATE(FPurchaseConfirmationFailure);
+
+	/**
+		Delegate used to report a successful item purchase history retrieval.
+		@param PurchasesHistory The current purchases history of the app.
+	*/
+	DECLARE_DYNAMIC_MULTICAST_DELEGATE_OneParam(FGetPurchaseHistorySuccess, const TArray<FPurchaseConfirmation>&, PurchaseHistory);
+
+	/**
+		Delegate used to pass a purchase history request failure back to the instigating blueprint.
+	*/
+	DECLARE_DYNAMIC_MULTICAST_DELEGATE(FGetPurchaseHistoryFailure);
+
+	/** Activated when a log message is generated on the capture thread. */
+	FInAppPurchaseLogMessage& OnInAppPurchaseLogMessage();
+
+	/** Activated when retrieval of items details succeeds. */
+	FGetItemsDetailsSuccess& OnGetItemsDetailsSuccess();
+
+	/** Activated when retrieval of items details fails. */
+	FGetItemsDetailsFailure& OnGetItemsDetailsFailure();
+
+	/** Activated when a purchase confirmationn request succeeds. */
+	FPurchaseConfirmationSuccess& OnPurchaseConfirmationSuccess();
+
+	/** Activated when a purchase confirmation request fails. */
+	FPurchaseConfirmationFailure& OnPurchaseConfirmationFailure();
+
+	/** Activated when a purchase history request succeeds. */
+	FGetPurchaseHistorySuccess& OnGetPurchaseHistorySuccess();
+
+	/** Activated when a purchase history request fails. */
+	FGetPurchaseHistoryFailure& OnGetPurchaseHistoryFailure();
+
+	void Log(const FString& LogMessage, bool bError = true);
+
+	UPROPERTY(BlueprintAssignable, Category = "InAppPurchase|MagicLeap", meta = (AllowPrivateAccess = true))
+	FInAppPurchaseLogMessage InAppPurchaseLogMessage;
+
+	UPROPERTY(BlueprintAssignable, Category = "InAppPurchase|MagicLeap", meta = (AllowPrivateAccess = true))
+	FGetItemsDetailsSuccess GetItemsDetailsSuccess;
+
+	UPROPERTY(BlueprintAssignable, Category = "InAppPurchase|MagicLeap", meta = (AllowPrivateAccess = true))
+	FGetItemsDetailsFailure GetItemsDetailsFailure;
+
+	UPROPERTY(BlueprintAssignable, Category = "InAppPurchase|MagicLeap", meta = (AllowPrivateAccess = true))
+	FPurchaseConfirmationSuccess PurchaseConfirmationSuccess;
+
+	UPROPERTY(BlueprintAssignable, Category = "InAppPurchase|MagicLeap", meta = (AllowPrivateAccess = true))
+	FPurchaseConfirmationFailure PurchaseConfirmationFailure;
+
+	UPROPERTY(BlueprintAssignable, Category = "InAppPurchase|MagicLeap", meta = (AllowPrivateAccess = true))
+	FGetPurchaseHistorySuccess GetPurchaseHistorySuccess;
+
+	UPROPERTY(BlueprintAssignable, Category = "InAppPurchase|MagicLeap", meta = (AllowPrivateAccess = true))
+	FGetPurchaseHistoryFailure GetPurchaseHistoryFailure;
+
+private:
+	class FInAppPurchaseImpl* Impl;
+};
+
+DECLARE_LOG_CATEGORY_EXTERN(LogInAppPurchase, Verbose, All);
\ No newline at end of file
diff --git a/Engine/Plugins/Lumin/MagicLeap/Source/MagicLeap/Public/LuminApplicationLifecycleComponent.h b/Engine/Plugins/Lumin/MagicLeap/Source/MagicLeap/Public/LuminApplicationLifecycleComponent.h
new file mode 100644
--- /dev/null
+++ b/Engine/Plugins/Lumin/MagicLeap/Source/MagicLeap/Public/LuminApplicationLifecycleComponent.h
@@ -0,0 +1,45 @@
+// Copyright 1998-2019 Epic Games, Inc. All Rights Reserved.
+
+#pragma once
+
+#include "Components/ApplicationLifecycleComponent.h"
+#include "LuminApplicationLifecycleComponent.generated.h"
+
+/** Component to handle receiving notifications from the LuminOS about application state (activated, suspended, termination, standby etc). */
+UCLASS(ClassGroup=Utility, HideCategories=(Activation, "Components|Activation", Collision), meta=(BlueprintSpawnableComponent))
+class MAGICLEAP_API ULuminApplicationLifecycleComponent : public UApplicationLifecycleComponent
+{
+	GENERATED_BODY()
+
+public:
+	DECLARE_DYNAMIC_MULTICAST_DELEGATE(FLuminApplicationLifetimeDelegate);
+
+	void OnRegister() override;
+	void OnUnregister() override;
+
+	/**
+		This event is called when the device has transitioned to the active mode
+		from reality or standby.
+		This is triggered when the device comes out of the reality mode
+		or when the wearable is back on head and is no longer in standby mode.
+	*/
+	UPROPERTY(BlueprintAssignable)
+	FLuminApplicationLifetimeDelegate DeviceHasReactivatedDelegate;  
+
+	/** This event is called when the device has transitioned to the reality mode. */
+	UPROPERTY(BlueprintAssignable)
+	FLuminApplicationLifetimeDelegate DeviceWillEnterRealityModeDelegate;  
+
+	/**
+		This callback is called when the device has transitioned to the standby mode.
+		This is triggered when the wearable is off head.
+	*/
+	UPROPERTY(BlueprintAssignable)
+	FLuminApplicationLifetimeDelegate DeviceWillGoInStandbyDelegate;  
+
+private:
+	/** Native handlers that get registered with the actual FLuminDelegates, and then proceed to broadcast to the delegates above */
+	void DeviceHasReactivatedDelegate_Handler() { DeviceHasReactivatedDelegate.Broadcast(); }
+	void DeviceWillEnterRealityModeDelegate_Handler() { DeviceWillEnterRealityModeDelegate.Broadcast(); }
+	void DeviceWillGoInStandbyDelegate_Handler() { DeviceWillGoInStandbyDelegate.Broadcast(); }
+};
diff --git a/Engine/Plugins/Lumin/MagicLeap/Source/MagicLeap/Public/MagicLeapHMDFunctionLibrary.h b/Engine/Plugins/Lumin/MagicLeap/Source/MagicLeap/Public/MagicLeapHMDFunctionLibrary.h
new file mode 100644
--- /dev/null
+++ b/Engine/Plugins/Lumin/MagicLeap/Source/MagicLeap/Public/MagicLeapHMDFunctionLibrary.h
@@ -0,0 +1,128 @@
+// Copyright 1998-2019 Epic Games, Inc. All Rights Reserved.
+
+#pragma once
+
+#include "CoreMinimal.h"
+#include "UObject/Object.h"
+#include "Kismet/BlueprintFunctionLibrary.h"
+#include "MagicLeapHMDFunctionLibrary.generated.h"
+
+UENUM(BlueprintType)
+enum class EMagicLeapHeadTrackingError : uint8
+{
+	None,
+	NotEnoughFeatures,
+	LowLight,
+	Unknown
+};
+
+UENUM(BlueprintType)
+enum class EMagicLeapHeadTrackingMode : uint8
+{
+	PositionAndOrientation,
+	Unavailable,
+	Unknown
+};
+
+/** Different types of map events that can occur that a developer may have to handle. */
+UENUM(BlueprintType)
+enum class EMagicLeapHeadTrackingMapEvent : uint8
+{
+	/** Map was lost. It could possibly recover. */
+	Lost,
+	/** Previous map was recovered. */
+	Recovered,
+	/** Failed to recover previous map. */
+	RecoveryFailed,
+	/** New map session created. */
+	NewSession
+};
+
+USTRUCT(BlueprintType)
+struct FMagicLeapHeadTrackingState
+{
+	GENERATED_BODY()
+
+public:
+	UPROPERTY(EditAnywhere, BlueprintReadWrite, Category = "MagicLeap")
+	EMagicLeapHeadTrackingMode Mode;
+
+	UPROPERTY(EditAnywhere, BlueprintReadWrite, Category = "MagicLeap")
+	EMagicLeapHeadTrackingError Error;
+
+	UPROPERTY(EditAnywhere, BlueprintReadWrite, Category = "MagicLeap")
+	float Confidence;
+};
+
+UCLASS(ClassGroup = MagicLeap)
+class MAGICLEAP_API UMagicLeapHMDFunctionLibrary : public UBlueprintFunctionLibrary
+{
+	GENERATED_BODY()
+
+public:
+	UFUNCTION(BlueprintCallable, Category = "Input|MagicLeap", meta = (DeprecatedFunction, DeprecationMessage = "Use XRPawn setup for coordinate space calibration"))
+	static void SetBasePosition(const FVector& InBasePosition);
+
+	UFUNCTION(BlueprintCallable, Category = "Input|MagicLeap", meta = (DeprecatedFunction, DeprecationMessage = "Use XRPawn setup for coordinate space calibration"))
+	static void SetBaseOrientation(const FQuat& InBaseOrientation);
+
+	UFUNCTION(BlueprintCallable, Category = "Input|MagicLeap", meta = (DeprecatedFunction, DeprecationMessage = "Use XRPawn setup for coordinate space calibration"))
+	static void SetBaseRotation(const FRotator& InBaseRotation);
+
+   /**
+	* Set the actor whose location is used as the focus point. The focus distance is the distance from the HMD to the focus point.
+	*
+	* @param InFocusActor			The actor that will be set as the new focus actor.
+	* @param bSetStabilizationActor  True if the function should set the stabilization depth actor to match the passed in focus actor. (RECOMMENDED TO STAY CHECKED)
+	*/
+	UFUNCTION(BlueprintCallable, Category = "Rendering|MagicLeap")
+	static void SetFocusActor(const AActor* InFocusActor, bool bSetStabilizationActor = true);
+
+   /**
+	* Set the actor whose location is used as the depth for timewarp stabilization.
+	*
+	* @param InStabilizationDepthActor  The actor that will be set as the new stabilization depth actor.
+	* @param bSetFocusActor				True if the function should set the focus actor to match the passed in stabilization depth actor. (RECOMMENDED TO STAY CHECKED)
+	*/
+	UFUNCTION(BlueprintCallable, Category = "Rendering|MagicLeap")
+	static void SetStabilizationDepthActor(const AActor* InStabilizationDepthActor, bool bSetFocusActor = true);
+
+	UFUNCTION(BlueprintCallable, BlueprintPure, Category = "MagicLeap")
+	static int32 GetMLSDKVersionMajor();
+
+	UFUNCTION(BlueprintCallable, BlueprintPure, Category = "MagicLeap")
+	static int32 GetMLSDKVersionMinor();
+
+	UFUNCTION(BlueprintCallable, BlueprintPure, Category = "MagicLeap")
+	static int32 GetMLSDKVersionRevision();
+
+	UFUNCTION(BlueprintCallable, BlueprintPure, Category = "MagicLeap")
+	static FString GetMLSDKVersion();
+
+	/** Returns true if this code is executing on the ML HMD, false otherwise (e.g. it's executing on PC) */
+	UFUNCTION(BlueprintCallable, BlueprintPure, Category = "MagicLeap")
+	static bool IsRunningOnMagicLeapHMD();
+
+	UFUNCTION(BlueprintCallable, BlueprintPure, Category = "MagicLeap")
+	static bool GetHeadTrackingState(FMagicLeapHeadTrackingState& State);
+
+	/**
+		Get map events.
+		A developer must be aware of certain events that can occur under degenerative conditions
+		in order to cleanly handle it. The most important event to be aware of is when a map changes.
+		In the case that a new map session begins, or recovery fails, all formerly cached transform
+		and world reconstruction data (raycast, planes, mesh) is invalidated and must be updated.
+		@param MapEvents Set of map events occured since the last frame. Valid only if function returns true.
+		@return true if call get map events from the platform succeeded, false otherwise.
+	 */
+	UFUNCTION(BlueprintCallable, BlueprintPure, Category = "MagicLeap")
+	static bool GetHeadTrackingMapEvents(TSet<EMagicLeapHeadTrackingMapEvent>& MapEvents);
+
+	/**
+		Notifies lifecycle that the app is ready to run (dismisses the loading logo).
+		@note This MUST be called if you have checked bManualCallToAppReady in LuminRuntimeSettings.  Failure to do so will cause the 
+		application to remain in the loading state.
+	*/
+	UFUNCTION(BlueprintCallable, Category = "MagicLeap")
+	static bool SetAppReady();
+};
diff --git a/Engine/Plugins/Lumin/MagicLeap/Source/MagicLeap/Public/MagicLeapHeadTrackingNotificationsComponent.h b/Engine/Plugins/Lumin/MagicLeap/Source/MagicLeap/Public/MagicLeapHeadTrackingNotificationsComponent.h
new file mode 100644
--- /dev/null
+++ b/Engine/Plugins/Lumin/MagicLeap/Source/MagicLeap/Public/MagicLeapHeadTrackingNotificationsComponent.h
@@ -0,0 +1,53 @@
+// Copyright 1998-2019 Epic Games, Inc. All Rights Reserved.
+
+#pragma once
+
+#include "VRNotificationsComponent.h"
+#include "MagicLeapHMDFunctionLibrary.h"
+#include "MagicLeapHeadTrackingNotificationsComponent.generated.h"
+
+/** 
+  Provides head tracking map events to enable apps to cleanly handle it. 
+  The most important event to be aware of is when a map changes.
+  In the case that a new map session begins, or recovery fails, all formerly cached transform
+  and world reconstruction data (raycast, planes, mesh) is invalidated and must be updated.
+*/
+UCLASS(ClassGroup = MagicLeap, BlueprintType, Blueprintable, meta = (BlueprintSpawnableComponent))
+class MAGICLEAP_API UMagicLeapHeadTrackingNotificationsComponent : public UVRNotificationsComponent
+{
+	GENERATED_BODY()
+
+public:
+	UMagicLeapHeadTrackingNotificationsComponent();
+
+	virtual void TickComponent(float DeltaTime, enum ELevelTick TickType, FActorComponentTickFunction *ThisTickFunction) override;
+
+	void BindToOnHeadTrackingLostEvent(const FVRNotificationsDelegate& EventDelegate);
+	void BindToOnHeadTrackingRecoveredEvent(const FVRNotificationsDelegate& EventDelegate);
+	void BindToOnHeadTrackingRecoveryFailedEvent(const FVRNotificationsDelegate& EventDelegate);
+	void BindToOnHeadTrackingNewSessionStartedEvent(const FVRNotificationsDelegate& EventDelegate);
+
+	void UnbindOnHeadTrackingLostEvent();
+	void UnbindOnHeadTrackingRecoveredEvent();
+	void UnbindOnHeadTrackingRecoveryFailedEvent();
+	void UnbindOnHeadTrackingNewSessionStartedEvent();
+
+private:
+	/** Map was lost. It could possibly recover. */
+	UPROPERTY(BlueprintAssignable)
+	FVRNotificationsDelegate OnHeadTrackingLost;
+
+	/** Previous map was recovered. */
+	UPROPERTY(BlueprintAssignable, meta = (AllowPrivateAccess = true))
+	FVRNotificationsDelegate OnHeadTrackingRecovered;
+
+	/** Failed to recover previous map. */
+	UPROPERTY(BlueprintAssignable, meta = (AllowPrivateAccess = true))
+	FVRNotificationsDelegate OnHeadTrackingRecoveryFailed;
+
+	/** New map session created. */
+	UPROPERTY(BlueprintAssignable, meta = (AllowPrivateAccess = true))
+	FVRNotificationsDelegate OnHeadTrackingNewSessionStarted;
+
+	TSet<EMagicLeapHeadTrackingMapEvent> PreviousMapEvents;
+};
diff --git a/Engine/Plugins/Lumin/MagicLeap/Source/MagicLeap/Public/MagicLeapRunnable.h b/Engine/Plugins/Lumin/MagicLeap/Source/MagicLeap/Public/MagicLeapRunnable.h
new file mode 100644
--- /dev/null
+++ b/Engine/Plugins/Lumin/MagicLeap/Source/MagicLeap/Public/MagicLeapRunnable.h
@@ -0,0 +1,192 @@
+// Copyright 1998-2019 Epic Games, Inc. All Rights Reserved.
+
+#pragma once
+
+#include "HAL/Runnable.h"
+#include "HAL/RunnableThread.h"
+#include "HAL/ThreadSafeBool.h"
+#include "HAL/ThreadSafeCounter64.h"
+#include "HAL/Event.h"
+#include "Containers/Queue.h"
+#include "Containers/Array.h"
+#include "Misc/Variant.h"
+#include "AppEventHandler.h"
+#include "Lumin/CAPIShims/LuminAPI.h"
+
+struct FMagicLeapTask
+{
+	bool bSuccess;
+
+	FMagicLeapTask()
+	: bSuccess(false)
+	{
+	}
+};
+
+template <class TTaskType>
+class FMagicLeapRunnable : public FRunnable
+{
+public:
+	FMagicLeapRunnable(const TArray<EMagicLeapPrivilege>& InRequiredPrivileges, const FString& InName, EThreadPriority InThreadPriority = EThreadPriority::TPri_BelowNormal)
+	: AppEventHandler(InRequiredPrivileges)
+	, Name(InName)
+	, Thread(nullptr)
+	, ThreadPriority(InThreadPriority)
+	, StopTaskCounter(0)
+	, Semaphore(nullptr)
+	, bPaused(false)
+	{
+		AppEventHandler.SetOnAppPauseHandler([this]() 
+		{
+			OnAppPause();
+		});
+
+		AppEventHandler.SetOnAppResumeHandler([this]()
+		{
+			OnAppResume();
+		});
+
+		AppEventHandler.SetOnAppShutDownHandler([this]()
+		{
+			OnAppShutDown();
+		});
+	}
+
+	virtual ~FMagicLeapRunnable()
+	{
+		StopTaskCounter.Increment();
+
+		if (Semaphore)
+		{
+			Semaphore->Trigger();
+			Thread->WaitForCompletion();
+			FGenericPlatformProcess::ReturnSynchEventToPool(Semaphore);
+			Semaphore = nullptr;
+		}
+
+		delete Thread;
+		Thread = nullptr;
+	}
+
+	uint32 Run() override
+	{
+		while (StopTaskCounter.GetValue() == 0)
+		{
+			if (bPaused)
+			{
+				Pause();
+				// Cancel any incoming tasks.
+				CancelIncomingTasks();
+				// Wait for signal from resume call.
+				Semaphore->Wait();
+				Resume();
+			}
+			else if (!IncomingTasks.IsEmpty())
+			{
+				DoNextTask();
+			}
+			else
+			{
+				Semaphore->Wait();
+			}
+		}
+
+		return 0;
+	}
+
+	void OnAppPause()
+	{
+		bPaused = true;
+		if (Semaphore)
+		{
+			Semaphore->Trigger();
+		}
+	}
+
+	void OnAppResume()
+	{
+		bPaused = false;
+		if (Semaphore)
+		{
+			Semaphore->Trigger();
+		}
+	}
+
+	void OnAppShutDown()
+	{
+		Stop();
+	}
+
+	void PushNewTask(TTaskType InTask)
+	{
+		// on demand thread creation
+		if (!Thread)
+		{
+			Semaphore = FGenericPlatformProcess::GetSynchEventFromPool(false);
+			Thread = FRunnableThread::Create(this, *Name, 0, ThreadPriority, FPlatformAffinity::GetPoolThreadMask());
+		}
+
+		IncomingTasks.Enqueue(InTask);
+		// wake up the worker to process the task
+		Semaphore->Trigger();
+	}
+
+	void PushCompletedTask(TTaskType InTask)
+	{
+		CompletedTasks.Enqueue(InTask);
+	}
+
+	bool TryGetCompletedTask(TTaskType& OutCompletedTask)
+	{
+#if PLATFORM_LUMIN
+		if (CompletedTasks.Peek(OutCompletedTask))
+		{
+			CompletedTasks.Pop();
+			return true;
+		}
+#endif //PLATFORM_LUMIN
+		return false;
+	}
+
+protected:
+	void CancelIncomingTasks()
+	{
+		while (IncomingTasks.Dequeue(CurrentTask))
+		{
+			CurrentTask.bSuccess = false;
+			CompletedTasks.Enqueue(CurrentTask);
+		}
+	}
+
+	virtual bool ProcessCurrentTask() = 0;
+	virtual void Pause() {}
+	virtual void Resume() {}
+
+	/** Internal thread this instance is running on */
+	MagicLeap::IAppEventHandler AppEventHandler;
+	FString Name;
+	FRunnableThread* Thread;
+	EThreadPriority ThreadPriority;
+	FThreadSafeCounter StopTaskCounter;
+	FEvent* Semaphore;
+	FThreadSafeBool bPaused;
+	TQueue<TTaskType, EQueueMode::Spsc> IncomingTasks;
+	TQueue<TTaskType, EQueueMode::Spsc> CompletedTasks;
+	TTaskType CurrentTask;
+
+private:
+	bool DoNextTask()
+	{
+		IncomingTasks.Dequeue(CurrentTask);
+		CurrentTask.bSuccess = ProcessCurrentTask();
+
+		if (bPaused)
+		{
+			return false;
+		}
+
+		CompletedTasks.Enqueue(CurrentTask);
+
+		return CurrentTask.bSuccess;
+	}
+};
diff --git a/Engine/Plugins/Lumin/MagicLeap/Source/MagicLeap/Public/MagicLeapSettings.h b/Engine/Plugins/Lumin/MagicLeap/Source/MagicLeap/Public/MagicLeapSettings.h
new file mode 100644
--- /dev/null
+++ b/Engine/Plugins/Lumin/MagicLeap/Source/MagicLeap/Public/MagicLeapSettings.h
@@ -0,0 +1,28 @@
+// Copyright 1998-2019 Epic Games, Inc. All Rights Reserved.
+
+#pragma once
+
+#include "CoreMinimal.h"
+#include "MagicLeapSettings.generated.h"
+
+/**
+* Implements the settings for the Magic Leap SDK setup.
+*/
+UCLASS(config=Engine, defaultconfig)
+class MAGICLEAP_API UMagicLeapSettings : public UObject
+{
+public:
+	GENERATED_BODY()
+
+	// Enables 'Zero Iteration mode'. Note: Vulkan rendering will be used by default. Set bUseVulkan to false to use OpenGL instead.
+	UPROPERTY(GlobalConfig, EditAnywhere, Category = "General", Meta = (DisplayName = "Enable Zero Iteration", ConfigRestartRequired = true))
+	bool bEnableZI;
+
+	// Use the editor in Vulkan. If False, OpenGL is used with ZI.
+	UPROPERTY(GlobalConfig, Meta = (ConfigRestartRequired = true))
+	bool bUseVulkanForZI;
+
+	// Use the MagicLeapAudio mixer module when using ZI. This will play audio via the ML device. Otherwise, the audio is played on the host machine itself.
+	UPROPERTY(GlobalConfig, EditAnywhere, Category = "General", Meta = (DisplayName = "Use MLAudio for Zero Iteration (host audio otherwise)", ConfigRestartRequired = true))
+	bool bUseMLAudioForZI;
+};
diff --git a/Engine/Plugins/Lumin/MagicLeap/Source/MagicLeap/Public/MeshBlockSelectorInterface.h b/Engine/Plugins/Lumin/MagicLeap/Source/MagicLeap/Public/MeshBlockSelectorInterface.h
new file mode 100644
--- /dev/null
+++ b/Engine/Plugins/Lumin/MagicLeap/Source/MagicLeap/Public/MeshBlockSelectorInterface.h
@@ -0,0 +1,27 @@
+// Copyright 1998-2019 Epic Games, Inc. All Rights Reserved.
+#pragma once
+
+#include "Engine/Engine.h"
+#include "MeshTrackerTypes.h"
+#include "MeshBlockSelectorInterface.generated.h"
+
+UINTERFACE(BlueprintType, Blueprintable)
+class MAGICLEAP_API UMagicLeapMeshBlockSelectorInterface : public UInterface
+{
+	GENERATED_BODY()
+};
+
+/** Interface to select blocks for a mesh request. */
+class MAGICLEAP_API IMagicLeapMeshBlockSelectorInterface
+{
+	GENERATED_BODY()
+
+public:
+	/**
+	 * Given the new mesh information, select the blocks you want to keep and the level of detail for each of those blocks. 
+	 * @param NewMeshInfo Information on the latest mesh blocks available.
+	 * @param RequestedMesh output array containing details of blocks for whom the mesh should be requested.
+	 */
+	UFUNCTION(BlueprintCallable, BlueprintNativeEvent, Category="Meshing|MagicLeap")
+	void SelectMeshBlocks(const FMagicLeapTrackingMeshInfo& NewMeshInfo, TArray<FMagicLeapMeshBlockRequest>& RequestedMesh);
+};
diff --git a/Engine/Plugins/Lumin/MagicLeap/Source/MagicLeap/Public/MeshTrackerComponent.h b/Engine/Plugins/Lumin/MagicLeap/Source/MagicLeap/Public/MeshTrackerComponent.h
new file mode 100644
--- /dev/null
+++ b/Engine/Plugins/Lumin/MagicLeap/Source/MagicLeap/Public/MeshTrackerComponent.h
@@ -0,0 +1,183 @@
+// Copyright 1998-2019 Epic Games, Inc. All Rights Reserved.
+
+#pragma once
+
+#include "MeshTrackerTypes.h"
+#include "MeshBlockSelectorInterface.h"
+#include "Components/SceneComponent.h"
+#include "MeshTrackerComponent.generated.h"
+
+/**
+	The MeshTrackerComponent class manages requests for environmental mesh data, processes the results and provides
+	them to the calling system. The calling system is able request environmental mesh data within a specified area.
+	Various other search criteria can be set via this class's public properties.  Mesh data requests are processed
+	on a separate thread.  Once a mesh data request has been processed the calling system will be notified via an
+	FOnMeshTrackerUpdated broadcast.
+*/
+UCLASS(ClassGroup = MagicLeap, BlueprintType, Blueprintable, EditInlineNew, meta = (BlueprintSpawnableComponent))
+class MAGICLEAP_API UMeshTrackerComponent : public USceneComponent, public IMagicLeapMeshBlockSelectorInterface
+{
+	GENERATED_BODY()
+
+public:
+	UMeshTrackerComponent();
+	/** Destroys the FMagicLeapMeshTrackerImpl instance.*/
+	~UMeshTrackerComponent();
+
+	/**
+		Sets the procedural mesh component that will store and display the environmental mesh results.
+		@param InMRMeshPtr The procedural mesh component to store the query result in.
+	*/
+	UFUNCTION(BlueprintCallable, Category = "Mesh Reconstruction")
+	void ConnectMRMesh(class UMRMeshComponent* InMRMeshPtr);
+
+	/**
+		Unlinks the current procedural mesh component from the mesh tracking system.
+		@param InMRMeshPtr The procedural mesh component to unlink from the mesh tracking system.
+	*/
+	UFUNCTION(BlueprintCallable, Category = "Mesh Reconstruction")
+	void DisconnectMRMesh(class UMRMeshComponent* InMRMeshPtr);
+
+	/**
+		Sets the interface to be used for selecting blocks to mesh.
+		@param Selector pointer to the object which implements IMagicLeapMeshBlockSelectorInterface.
+	*/
+	UFUNCTION(BlueprintCallable, Category = "Meshing|MagicLeap")
+	void ConnectBlockSelector(TScriptInterface<IMagicLeapMeshBlockSelectorInterface> Selector);
+
+	/** 
+	 * Disconnects the previously connected IMagicLeapMeshBlockSelectorInterface.
+	 * The default implementation is used this case - all new and updated blocks are meshed with the MeshTrackerComponent'd LevelOfDetail.
+	 */
+	UFUNCTION(BlueprintCallable, Category = "Meshing|MagicLeap")
+	void DisconnectBlockSelector();
+
+	/**
+	 * Gets the number of mesh bricks in the pending queue.
+	 */
+	UFUNCTION(BlueprintCallable, Category = "Meshing|MagicLeap")
+	int32 GetNumQueuedBlockUpdates();
+
+	/** IMagicLeapMeshBlockSelectorInterface */
+	UFUNCTION(BlueprintCallable, BlueprintNativeEvent, Category="Meshing|MagicLeap")
+	void SelectMeshBlocks(const FMagicLeapTrackingMeshInfo& NewMeshInfo, TArray<FMagicLeapMeshBlockRequest>& RequestedMesh);
+	virtual void SelectMeshBlocks_Implementation(const FMagicLeapTrackingMeshInfo& NewMeshInfo, TArray<FMagicLeapMeshBlockRequest>& RequestedMesh) override;
+
+	/**
+		Delegate used by OnMeshUpdated().
+		@param ID The ID of the mesh section in the Block that was updated.
+		@param Vertices List of all vertices in the updated mesh section.
+		@param Triangles List of all triangles in the updated mesh section.
+		@param Normals List of the normals of all triangles in the updated mesh section.
+		@param Confidence List of the confidence values per vertex in the updated mesh section. This can be used to determine if the user needs to scan more.
+	*/
+	DECLARE_DYNAMIC_MULTICAST_DELEGATE_FiveParams(FOnMeshTrackerUpdated, FGuid, ID, const TArray<FVector>&, Vertices, const TArray<int32>&, Triangles, const TArray<FVector>&, Normals, const TArray<float>&, Confidence);
+
+	/** Activated whenever new information about this mesh tracker is detected. */
+	UPROPERTY(BlueprintAssignable)
+	FOnMeshTrackerUpdated OnMeshTrackerUpdated;
+
+	/** Set to true to start scanning the world for meshes. */
+	UPROPERTY(EditAnywhere, BlueprintReadWrite, Category = "Meshing|MagicLeap")
+	bool ScanWorld = true;
+
+	/** The type of mesh to query. */
+	UPROPERTY(EditAnywhere, BlueprintReadWrite, Category = "Meshing|MagicLeap")
+	EMagicLeapMeshType MeshType = EMagicLeapMeshType::Triangles;
+
+	/** Bounding box for the mesh scan. The mesh will be scanned for only within this box. */
+	UPROPERTY(EditAnywhere, BlueprintReadWrite, Category = "Meshing|MagicLeap")
+	class UBoxComponent* BoundingVolume;
+
+	/** Meshing LOD. */
+	UPROPERTY(EditAnywhere, BlueprintReadWrite, Category = "Meshing|MagicLeap")
+	EMagicLeapMeshLOD LevelOfDetail = EMagicLeapMeshLOD::Medium;
+
+	/** The perimeter (in Unreal Units) of gaps to be filled. 0 means do not fill. A good value is 300cm. */
+	UPROPERTY(EditAnywhere, BlueprintReadWrite, Category = "Meshing|MagicLeap")
+	float PerimeterOfGapsToFill = 300.0f;
+
+	/** If true, the system will planarize the returned mesh i.e. planar regions will be smoothed out. */
+	UPROPERTY(EditAnywhere, BlueprintReadWrite, Category = "Meshing|MagicLeap")
+	bool Planarize = false;
+
+	/** Any section that is disconnected from the main mesh and has an area (in Unreal Units squared) 
+	    less than this value will be removed. 
+		0 means do not remove disconnected sections. A good value is 50cm. */
+	UPROPERTY(EditAnywhere, BlueprintReadWrite, Category = "Meshing|MagicLeap")
+	float DisconnectedSectionArea = 50.0f;
+
+	/** If true, the system will generate normals for the triangle vertices. */
+	UPROPERTY(EditAnywhere, BlueprintReadWrite, Category = "Meshing|MagicLeap")
+	bool RequestNormals = true;
+
+	/**
+		If true, the system will generate the mesh confidence values for the triangle vertices.
+		These confidence values can be used to determine if the user needs to scan more.
+	*/
+	UPROPERTY(EditAnywhere, BlueprintReadWrite, Category = "Meshing|MagicLeap")
+	bool RequestVertexConfidence = false;
+
+	/**
+	* Vertex Colors can be unused, or filled with several types of information.
+	*/
+	UPROPERTY(EditAnywhere, BlueprintReadWrite, Category = "Meshing|MagicLeap")
+	EMagicLeapMeshVertexColorMode VertexColorMode = EMagicLeapMeshVertexColorMode::None;
+
+	/** Colors through which we cycle when setting vertex color by block. */
+	UPROPERTY(EditAnywhere, BlueprintReadWrite, Category = "Meshing|MagicLeap")
+	TArray<FColor> BlockVertexColors;
+
+	/** Color mapped to confidence value of zero. */
+	UPROPERTY(EditAnywhere, BlueprintReadWrite, Category = "Meshing|MagicLeap")
+	FLinearColor VertexColorFromConfidenceZero;
+
+	/** Color mapped to confidence value of one. */
+	UPROPERTY(EditAnywhere, BlueprintReadWrite, Category = "Meshing|MagicLeap")
+	FLinearColor VertexColorFromConfidenceOne;
+
+	/**
+		If true, overlapping area between two mesh blocks will be removed.
+		This field is only valid when the MeshType is Blocks.
+	*/
+	UPROPERTY(EditAnywhere, BlueprintReadWrite, Category = "Meshing|MagicLeap")
+	bool RemoveOverlappingTriangles = false;
+
+	/** MRMeshComponent can render and provide collision based on the Mesh data. */
+	UPROPERTY(transient)
+	class UMRMeshComponent* MRMesh;
+
+	/** Specifies the number of MR Mesh bricks to create per frame. This is to
+	    amortize the cost of brick creation across frames rather than creating
+		all N bricks delivered in response to each meshing query.
+		If the value is <=0 then all meshes will be pushed to MR Mesh as they
+		are provided by the system. */
+	UPROPERTY(EditAnywhere, BlueprintReadWrite, Category = "Meshing|MagicLeap")
+	int32 BricksPerFrame = 1;
+
+	/** Polls for and handles the results of the environmental mesh queries. */
+	virtual void TickComponent(float DeltaTime, enum ELevelTick TickType, FActorComponentTickFunction *ThisTickFunction) override;
+
+	/** Unlinks the current procedural mesh component from the mesh tracking system. */
+	virtual void BeginDestroy() override;
+
+	/** Destroys the interface object to the mesh tracking api*/
+	virtual void FinishDestroy() override;
+
+#if WITH_EDITOR
+	virtual void PostEditChangeProperty(FPropertyChangedEvent& e) override;
+#endif
+
+private:
+	void RequestMeshInfo();
+	bool GetMeshInfoResult();
+	void RequestMesh();
+	bool GetMeshResult();
+
+	class FMagicLeapMeshTrackerImpl *Impl;
+
+#if WITH_EDITOR
+private:
+	void PrePIEEnded(bool bWasSimulatingInEditor);
+#endif
+};
diff --git a/Engine/Plugins/Lumin/MagicLeap/Source/MagicLeap/Public/MeshTrackerTypes.h b/Engine/Plugins/Lumin/MagicLeap/Source/MagicLeap/Public/MeshTrackerTypes.h
new file mode 100644
--- /dev/null
+++ b/Engine/Plugins/Lumin/MagicLeap/Source/MagicLeap/Public/MeshTrackerTypes.h
@@ -0,0 +1,122 @@
+// Copyright 1998-2019 Epic Games, Inc. All Rights Reserved.
+
+#pragma once
+
+#include "Engine/Engine.h"
+#include "CoreMinimal.h"
+#include "UObject/ObjectMacros.h"
+#include "MeshTrackerTypes.generated.h"
+
+/** Type of mesh to query from the underlying system. */
+UENUM(BlueprintType)
+enum class EMagicLeapMeshType : uint8
+{
+	/** Meshing should be done as triangles. */
+	Triangles,
+	/** Return mesh vertices as a point cloud. */
+	PointCloud
+};
+
+/** Vertex color mode. */
+UENUM(BlueprintType)
+enum class EMagicLeapMeshVertexColorMode : uint8
+{
+	/** Vertex Color is not set. */
+	None		UMETA(DisplayName = "No Vertex Color"),
+	/** Vertex confidence is interpolated between two specified colors. */
+	Confidence	UMETA(DisplayName = "Vertex Confidence"),
+	/** Each block is given a color from a list. */
+	Block		UMETA(DisplayName = "Blocks Colored"),
+	/** Each LOD is given a color from a list. */
+	LOD			UMETA(DisplayName = "LODs Colored")
+};
+
+/** Discrete level of detail required. */
+UENUM(BlueprintType)
+enum class EMagicLeapMeshLOD : uint8
+{
+	/** Minimum LOD. */
+	Minimum,
+	/** Medium LOD. */
+	Medium,
+	/** Maximum LOD. */
+	Maximum,
+};
+
+/** State of a block mesh. */
+UENUM(BlueprintType)
+enum class EMagicLeapMeshState : uint8
+{
+	/** Mesh has been created */
+	New,
+	/** Mesh has been updated. */
+	Updated,
+	/** Mesh has been deleted. */
+	Deleted,
+	/** Mesh is unchanged. */
+	Unchanged
+};
+
+/** Representation of a mesh block. */
+USTRUCT(BlueprintType)
+struct FMagicLeapMeshBlockInfo
+{
+	GENERATED_BODY()
+
+public:
+	/** The coordinate frame UID to represent the block. */
+	UPROPERTY(VisibleAnywhere, BlueprintReadOnly, Category = "Meshing|MagicLeap")
+	FGuid BlockID;
+
+	/** The center of the mesh block bounding box. */
+	UPROPERTY(VisibleAnywhere, BlueprintReadOnly, Category = "Meshing|MagicLeap")
+	FVector BlockPosition;
+
+	/** The orientation of the mesh block bounding box.*/
+	UPROPERTY(VisibleAnywhere, BlueprintReadOnly, Category = "Meshing|MagicLeap")
+	FRotator BlockOrientation;
+
+	/** The size of the mesh block bounding box (in Unreal Units). */
+	UPROPERTY(VisibleAnywhere, BlueprintReadOnly, Category = "Meshing|MagicLeap")
+	FVector BlockDimensions;
+
+	/** The timestamp when block was updated. */
+	UPROPERTY(VisibleAnywhere, BlueprintReadOnly, Category = "Meshing|MagicLeap")
+	FTimespan Timestamp;
+
+	/** The state of the mesh block. */
+	UPROPERTY(VisibleAnywhere, BlueprintReadOnly, Category = "Meshing|MagicLeap")
+	EMagicLeapMeshState BlockState;
+};
+
+/** Response structure for the mesh block info. */
+USTRUCT(BlueprintType)
+struct FMagicLeapTrackingMeshInfo
+{
+	GENERATED_BODY()
+
+public:
+	/** The response timestamp to a earlier request. */
+	UPROPERTY(VisibleAnywhere, BlueprintReadOnly, Category = "Meshing|MagicLeap")
+	FTimespan Timestamp;
+
+	/** The meshinfo returned by the system. */
+	UPROPERTY(VisibleAnywhere, BlueprintReadOnly, Category = "Meshing|MagicLeap")
+	TArray<FMagicLeapMeshBlockInfo> BlockData;
+};
+
+/** Request structure to get the actual mesh for a block. */
+USTRUCT(BlueprintType)
+struct FMagicLeapMeshBlockRequest
+{
+	GENERATED_BODY()
+
+public:
+	/** The UID to represent the block. */
+	UPROPERTY(EditAnywhere, BlueprintReadWrite, Category = "Meshing|MagicLeap")
+	FGuid BlockID;
+
+	/** The LOD level to request. */
+	UPROPERTY(EditAnywhere, BlueprintReadWrite, Category = "Meshing|MagicLeap")
+	EMagicLeapMeshLOD LevelOfDetail;
+};
