Copying //UE4/Partner-MagicLeap-Staging to Dev-VR (//UE4/Dev-VR)

diff --git a/Engine/Plugins/Lumin/MagicLeap/Source/MagicLeap/Public/RaycastComponent.h b/Engine/Plugins/Lumin/MagicLeap/Source/MagicLeap/Public/RaycastComponent.h
new file mode 100644
--- /dev/null
+++ b/Engine/Plugins/Lumin/MagicLeap/Source/MagicLeap/Public/RaycastComponent.h
@@ -0,0 +1,156 @@
+// Copyright 1998-2019 Epic Games, Inc. All Rights Reserved.
+
+#pragma once
+
+#include "Components/ActorComponent.h"
+#include "Kismet/BlueprintFunctionLibrary.h"
+#include "RaycastComponent.generated.h"
+
+/** Parameters for a raycast request. */
+USTRUCT(BlueprintType, meta = (HasNativeMake = "MagicLeap.MagicLeapRaycastFunctionLibrary.MakeRaycastQueryParams"))
+struct FMagicLeapRaycastQueryParams
+{
+	GENERATED_BODY()
+
+public:
+	/** Where the ray is cast from. */
+	UPROPERTY(EditAnywhere, BlueprintReadWrite, Category = "Raycast|MagicLeap")
+	FVector Position;
+
+	/** Direction of the ray to fire. */
+	UPROPERTY(EditAnywhere, BlueprintReadWrite, Category = "Raycast|MagicLeap")
+	FVector Direction = FVector(1, 0, 0);
+
+	/** Up vector of the ray to fire. This is used to orient the area the rays are cast over. */
+	UPROPERTY(EditAnywhere, BlueprintReadWrite, Category = "Raycast|MagicLeap")
+	FVector UpVector = FVector(0, 0, 1);
+
+	/** The number of horizontal rays. For single point raycast, set this to 1. */
+	UPROPERTY(EditAnywhere, BlueprintReadWrite, Category = "Raycast|MagicLeap", meta = (ClampMin = 1))
+	int32 Width = 1;
+
+	/** The number of vertical rays. For single point raycast, set this to 1. */
+	UPROPERTY(EditAnywhere, BlueprintReadWrite, Category = "Raycast|MagicLeap", meta = (ClampMin = 1))
+	int32 Height = 1;
+
+	/** The angular width, in degrees, over which the horizonal rays are evenly distributed to create a raycast area. */
+	UPROPERTY(EditAnywhere, BlueprintReadWrite, Category = "Raycast|MagicLeap")
+	float HorizontalFovDegrees;
+
+	/** If true, a ray will terminate when encountering an unobserved area and return a surface or
+	  the ray will continue until it ends or hits an observed surface. */
+	UPROPERTY(EditAnywhere, BlueprintReadWrite, Category = "Raycast|MagicLeap")
+	bool CollideWithUnobserved;
+
+	/** User data for this request. The same data will be included in the result for query identification. */
+	UPROPERTY(EditAnywhere, BlueprintReadWrite, Category = "Raycast|MagicLeap")
+	int32 UserData;
+};
+
+/** The states of a raycast result. */
+UENUM(BlueprintType)
+enum class EMagicLeapRaycastResultState : uint8
+{
+	/** The raycast request failed. */
+	RequestFailed,
+	/** The ray passed beyond maximum raycast distance and it doesn't hit any surface. */
+	NoCollision,
+	/** The ray hit an unobserved area. This can occur only when CollideWithUnobserved is set to true. */
+	HitUnobserved,
+	/** The ray hit an observed area. */
+	HitObserved
+};
+
+/** Result of a raycast. */
+USTRUCT(BlueprintType)
+struct FMagicLeapRaycastHitResult
+{
+	GENERATED_BODY()
+
+public:
+	/** The raycast result. If this field is either RequestFailed or NoCollision,
+	  most of the fields in this structure are invalid. */
+	UPROPERTY(BlueprintReadOnly, Category = "Raycast|MagicLeap")
+	EMagicLeapRaycastResultState HitState;
+
+	/** Where in the world the collision happened. This field is only valid if the state
+	  is either HitUnobserved or HitObserved. */
+	UPROPERTY(BlueprintReadOnly, Category = "Raycast|MagicLeap")
+	FVector HitPoint;
+
+	/** Normal to the surface where the ray collided. This field is only valid if the state
+	  is either HitUnobserved or HitObserved */
+	UPROPERTY(BlueprintReadOnly, Category = "Raycast|MagicLeap")
+	FVector Normal;
+
+	/** Confidence of the raycast result.
+
+	  Confidence is a non-negative value from 0 to 1 where closer to 1 indicates a higher quality.
+	  It will be an indication of how much error there is in the averaging.
+	  For example, a flat plane will have a high confidence, while if the surface was very noisy the confidence
+	  would be very low. This field is only valid if the state is either HitUnobserved or HitObserved.
+	*/
+	UPROPERTY(BlueprintReadOnly, Category = "Raycast|MagicLeap")
+	float Confidence;
+
+	/** The data set in the RaycastQueryParams. This can be used for query identification. */
+	UPROPERTY(BlueprintReadOnly, Category = "Raycast|MagicLeap")
+	int32 UserData;
+};
+
+/** Creates raycast requests and delegates their result. */
+UCLASS(ClassGroup = MagicLeap, BlueprintType, Blueprintable, EditInlineNew, meta = (BlueprintSpawnableComponent))
+class MAGICLEAP_API UMagicLeapRaycastComponent
+	: public UActorComponent
+{
+	GENERATED_BODY()
+
+public:
+	UMagicLeapRaycastComponent();
+	virtual ~UMagicLeapRaycastComponent();
+
+	/**
+	  Delegate used to convey the result of a raycast.
+	  @param HitResult structure containing the result of the raycast hit.
+	*/
+	DECLARE_DYNAMIC_DELEGATE_OneParam(FRaycastResultDelegate, FMagicLeapRaycastHitResult, HitResult);
+
+	/**
+	  Requests a raycast with the given query parameters.
+	  @param RequestParams Parameters for the raycast query.
+	  @param ResultDelegate Delegate which will be called when the raycast result is ready.
+	  @returns True if the raycast request was successfully placed, false otherwise.
+	*/
+	UFUNCTION(BlueprintCallable, Category = "Raycast|MagicLeap")
+	bool RequestRaycast(const FMagicLeapRaycastQueryParams& RequestParams, const FRaycastResultDelegate& ResultDelegate);
+
+	virtual void TickComponent(float DeltaTime, enum ELevelTick TickType, FActorComponentTickFunction *ThisTickFunction) override;
+	virtual void FinishDestroy() override;
+
+private:
+	struct FRaycastRequestMetaData
+	{
+	public:
+		FRaycastResultDelegate ResultDelegate;
+		int32 UserData;
+	};
+
+	TMap<uint64, FRaycastRequestMetaData> PendingRequests;
+	TArray<uint64> CompletedRequests;
+	class FMagicLeapRaycastTrackerImpl *Impl;
+
+#if WITH_EDITOR
+private:
+	void PrePIEEnded(bool bWasSimulatingInEditor);
+#endif
+};
+
+UCLASS(ClassGroup = MagicLeap)
+class MAGICLEAP_API UMagicLeapRaycastFunctionLibrary : public UBlueprintFunctionLibrary
+{
+	GENERATED_BODY()
+
+public:
+	UFUNCTION(BlueprintPure, BlueprintCallable, Category = "Raycast|MagicLeap", meta = (NativeMakeFunc, Direction="1,0,0", UpVector="0,0,1", Width="1", Height="1"))
+	static FMagicLeapRaycastQueryParams MakeRaycastQueryParams(FVector Position, FVector Direction, FVector UpVector, int32 Width, int32 Height, float HorizontalFovDegrees, bool CollideWithUnobserved, int32 UserData);
+};
diff --git a/Engine/Plugins/Lumin/MagicLeap/Source/MagicLeapAR/MagicLeapAR.Build.cs b/Engine/Plugins/Lumin/MagicLeap/Source/MagicLeapAR/MagicLeapAR.Build.cs
--- a/Engine/Plugins/Lumin/MagicLeap/Source/MagicLeapAR/MagicLeapAR.Build.cs
+++ b/Engine/Plugins/Lumin/MagicLeap/Source/MagicLeapAR/MagicLeapAR.Build.cs
@@ -14,8 +14,6 @@ namespace UnrealBuildTool.Rules
 					// ... add other private include paths required here ...
 				}
 			);
-			// This is not ideal but needs to be done in order to expose the private MagicLeapHMD header to this module.
-			//PrivateIncludePaths.Add(Path.Combine(new string[] { ModuleDirectory, "../../../../Lumin/MagicLeap", "Private" }));
 
 			PrivateDependencyModuleNames.AddRange(
 				new string[]
@@ -28,13 +26,11 @@ namespace UnrealBuildTool.Rules
 					"LuminRuntimeSettings",
 					"MagicLeap",
 					"MLSDK",
+					"MagicLeapPlanes",
 					"HeadMountedDisplay",
-					"AugmentedReality"
+					"AugmentedReality",
 				}
 			);
-
-			// This is not ideal but needs to be done in order to expose the private MagicLeapHMD header to this module.
-			PrivateIncludePaths.Add(Path.Combine(new string[] { ModuleDirectory, "..", "MagicLeap", "Private" }));
 		}
 	}
 }
diff --git a/Engine/Plugins/Lumin/MagicLeap/Source/MagicLeapAR/Private/LuminARAPI.cpp b/Engine/Plugins/Lumin/MagicLeap/Source/MagicLeapAR/Private/LuminARAPI.cpp
--- a/Engine/Plugins/Lumin/MagicLeap/Source/MagicLeapAR/Private/LuminARAPI.cpp
+++ b/Engine/Plugins/Lumin/MagicLeap/Source/MagicLeapAR/Private/LuminARAPI.cpp
@@ -7,21 +7,14 @@
 #include "Templates/Casts.h"
 
 #include "IMagicLeapPlugin.h"
-#include "MagicLeapPrivileges.h"
-#include "MagicLeapHMD.h"
-#include "AppFramework.h"
 #include "HeadMountedDisplayFunctionLibrary.h"
 
-#if WITH_MLSDK
-#pragma warning( push )
-#pragma warning( disable : 4201)
-#include "ml_planes.h"
-#pragma warning( pop ) 
-#endif //WITH_MLSDK
+#include "IMagicLeapPlanesModule.h"
+#include "MagicLeapPlanesFunctionLibrary.h"
 
 namespace
 {
-#if PLATFORM_LUMIN
+#if WITH_MLSDK
 	static const FMatrix LuminARToUnrealTransform = FMatrix(
 		FPlane(0.0f, 0.0f, -1.0f, 0.0f),
 		FPlane(1.0f, 0.0f, 0.0f, 0.0f),
@@ -55,7 +48,7 @@ namespace
 		OutPose.Pos = LuminARPoseMatrix.GetOrigin();
 		OutPose.Quat = LuminARPoseMatrix.ToQuat();
 	}
-#endif
+#endif // WITH_MLSDK
 	inline bool CheckIsSessionValid(FString TypeName, const TWeakPtr<FLuminARSession>& SessionPtr)
 	{
 		return true;
@@ -75,10 +68,8 @@ FLuminARSession::FLuminARSession()
 {
 	// Create Lumin ARSession handle.
 	LatestFrame = new FLuminARFrame(this);
-#if PLATFORM_LUMIN
 	InitTracker();
 	LatestFrame->Init();
-#endif
 }
 
 FLuminARSession::~FLuminARSession()
@@ -90,9 +81,7 @@ FLuminARSession::~FLuminARSession()
 
 	delete LatestFrame;
 
-#if PLATFORM_LUMIN
 	DestroyTracker();
-#endif
 }
 
 
@@ -106,29 +95,17 @@ float FLuminARSession::GetWorldToMeterScale()
 	return CachedWorldToMeterScale;
 }
 
-
-#if PLATFORM_ANDROID
-MLHandle FLuminARSession::GetPlaneTrackerHandle()
-{
-	return PlaneTrackerHandle;
-}
-#endif
-
 ELuminARAPIStatus FLuminARSession::Resume()
 {
 	ELuminARAPIStatus ResumeStatus = ELuminARAPIStatus::AR_SUCCESS;
-#if PLATFORM_LUMIN
 	//InitTracker();
-#endif
 	return ResumeStatus;
 }
 
 ELuminARAPIStatus FLuminARSession::Pause()
 {
 	ELuminARAPIStatus PauseStatus = ELuminARAPIStatus::AR_SUCCESS;
-#if PLATFORM_LUMIN
 	//DestroyTracker();
-#endif
 
 	for (UARPin* Anchor : UObjectManager->AllAnchors)
 	{
@@ -170,19 +147,23 @@ ELuminARAPIStatus FLuminARSession::CreateARAnchor(const FTransform& TransfromInT
 	ELuminARAPIStatus AnchorCreateStatus = ELuminARAPIStatus::AR_SUCCESS;
 	OutAnchor = nullptr;
 
-#if PLATFORM_LUMIN
 	TSharedPtr<LuminArAnchor> NewLuminArAnchor;
 	ArPose Pose;
+#if WITH_MLSDK
 	UnrealTransformToLuminARPose(TransfromInTrackingSpace, this, Pose, CachedWorldToMeterScale);
+#endif // WITH_MLSDK
+
 	if (TrackedGeometry == nullptr)
 	{
-		NewLuminArAnchor = MakeShared<LuminArAnchor>(Pose, ML_INVALID_HANDLE);
+#if WITH_MLSDK
+		NewLuminArAnchor = MakeShared<LuminArAnchor>(Pose, MagicLeap::MLHandleToFGuid(ML_INVALID_HANDLE));
+#endif // WITH_MLSDK
 	}
 	else
 	{
 		ensure(TrackedGeometry->GetNativeResource() != nullptr);
-		MLHandle ParentHandle = reinterpret_cast<FLuminARTrackableResource*>(TrackedGeometry->GetNativeResource())->GetNativeHandle();
-		ensure(ParentHandle != ML_INVALID_HANDLE);
+		const FGuid& ParentHandle = static_cast<FLuminARTrackableResource*>(TrackedGeometry->GetNativeResource())->GetNativeHandle();
+		ensure(MagicLeap::FGuidIsValidHandle(ParentHandle));
 		NewLuminArAnchor = MakeShared<LuminArAnchor>(Pose, ParentHandle);
 	}
 
@@ -196,7 +177,6 @@ ELuminARAPIStatus FLuminARSession::CreateARAnchor(const FTransform& TransfromInT
 		UObjectManager->AllAnchors.Add(OutAnchor);
 		UObjectManager->HandleToAnchorMap.Add(NewLuminArAnchor->Handle, OutAnchor);
 	}
-#endif
 	return AnchorCreateStatus;
 }
 
@@ -207,7 +187,6 @@ void FLuminARSession::DetachAnchor(UARPin* Anchor)
 		return;
 	}
 
-#if PLATFORM_LUMIN
 	LuminArAnchor* NativeResource = reinterpret_cast<LuminArAnchor*>(Anchor->GetNativeResource());
 	check(NativeResource);
 	NativeResource->Detach();
@@ -217,7 +196,6 @@ void FLuminARSession::DetachAnchor(UARPin* Anchor)
 
 	UObjectManager->HandleToAnchorMap.Remove(NativeResource->Handle);
 	UObjectManager->HandleToLuminAnchorMap.Remove(NativeResource->Handle);
-#endif
 	UObjectManager->AllAnchors.Remove(Anchor);
 }
 
@@ -228,36 +206,12 @@ void FLuminARSession::GetAllAnchors(TArray<UARPin*>& OutAnchors) const
 
 void FLuminARSession::InitTracker()
 {
-#if PLATFORM_LUMIN
-	if (!MLHandleIsValid(PlaneTrackerHandle))
-	{
-		MLResult PlaneCreateResult = MLPlanesCreate(&PlaneTrackerHandle);
-		if (PlaneCreateResult != MLResult_Ok)
-		{
-			UE_LOG(LogLuminARAPI, Warning, TEXT("Failed to create Plane Tracker for Lumin AR Session Result:%i"), PlaneCreateResult);
-		}
-	}
-	else
-	{
-		UE_LOG(LogLuminARAPI, Warning, TEXT("Tracker already exists"));
-	}
-#endif
+	UMagicLeapPlanesFunctionLibrary::CreateTracker();
 }
 
 void FLuminARSession::DestroyTracker()
 {
-#if PLATFORM_LUMIN
-	if (MLHandleIsValid(PlaneTrackerHandle))
-	{
-		MLResult Result = MLPlanesDestroy(PlaneTrackerHandle);
-		if (Result != MLResult_Ok)
-		{
-			UE_LOG(LogLuminARAPI, Warning, TEXT("Failed to destroy Plane Tracker for Lumin AR Session Result:%i"), Result);
-		}
-
-		PlaneTrackerHandle = ML_INVALID_HANDLE;
-	}
-#endif //WITH_MLSDK
+	UMagicLeapPlanesFunctionLibrary::DestroyTracker();
 }
 
 void FLuminARSession::AddReferencedObjects(FReferenceCollector& Collector)
@@ -277,31 +231,23 @@ FLuminARFrame::FLuminARFrame(FLuminARSession* InSession)
 	, LatestCameraTimestamp(0)
 	, LatestCameraTrackingState(ELuminARTrackingState::StoppedTracking)
 	, LatestARPlaneQueryStatus(ELuminARPlaneQueryStatus::Unknown)
-#if PLATFORM_LUMIN
-	, PlaneTrackerHandle(ML_INVALID_HANDLE)
-#endif
+	, bPlanesQueryPending(false)
 {
+	ResultDelegate.BindRaw(this, &FLuminARFrame::ProcessPlaneQuery);
 }
 
 FLuminARFrame::~FLuminARFrame()
 {
+	// TODO: @njain should we wait here if bPlanesQueryPending is true? Don't want the result delegate to come back to a deleted object.
 }
 
 void FLuminARFrame::Init()
 {
-#if PLATFORM_LUMIN
-	if (Session->GetPlaneTrackerHandle())
-	{
-		PlaneTrackerHandle = Session->GetPlaneTrackerHandle();
-}
-#endif
 }
 
-
 void FLuminARFrame::Update(float WorldToMeterScale)
 {
-#if PLATFORM_LUMIN
-	if (!MLHandleIsValid(PlaneTrackerHandle))
+	if (!UMagicLeapPlanesFunctionLibrary::IsTrackerValid())
 	{
 		LatestCameraTrackingState = ELuminARTrackingState::NotTracking;
 		return;
@@ -309,7 +255,6 @@ void FLuminARFrame::Update(float WorldToMeterScale)
 
 	// Update trackable that is cached in Unreal
 	StartPlaneQuery();
-	ProcessPlaneQuery();
 
 	switch (LatestARPlaneQueryStatus)
 	{
@@ -339,12 +284,12 @@ void FLuminARFrame::Update(float WorldToMeterScale)
 	UpdatedAnchors.Empty();
 	for (auto HandleToAnchorMapPair : Session->GetUObjectManager()->HandleToAnchorMap)
 	{
-		const MLHandle AnchorHandle = HandleToAnchorMapPair.Key;
+		const FGuid AnchorHandle = HandleToAnchorMapPair.Key;
 		UARPin* const AnchorPin = HandleToAnchorMapPair.Value;
 		LuminArAnchor* LuminArAnchor = reinterpret_cast<struct LuminArAnchor*>(AnchorPin->GetNativeResource());
 		check(LuminArAnchor);
-		const MLHandle ParentTrackableHandle = LuminArAnchor->ParentTrackable;
-		if (ParentTrackableHandle != ML_INVALID_HANDLE)
+		const FGuid ParentTrackableHandle = LuminArAnchor->ParentTrackable;
+		if (MagicLeap::FGuidIsValidHandle(ParentTrackableHandle))
 		{
 			UARTrackedGeometry* const ParentTrackable = Session->GetUObjectManager()->GetTrackableFromHandle<UARTrackedGeometry>(ParentTrackableHandle, Session);
 			const EARTrackingState AnchorTrackingState = ParentTrackable->GetTrackingState();
@@ -360,7 +305,6 @@ void FLuminARFrame::Update(float WorldToMeterScale)
 			UpdatedAnchors.Add(AnchorPin);
 		}
 	}
-#endif
 }
 
 //FTransform FLuminARFrame::GetCameraPose() const
@@ -385,7 +329,6 @@ void FLuminARFrame::GetUpdatedAnchors(TArray<UARPin*>& OutUpdatedAnchors) const
 
 void FLuminARFrame::ARLineTrace(FVector2D ScreenPosition, ELuminARLineTraceChannel RequestedTraceChannels, TArray<FARTraceResult>& OutHitResults) const
 {
-#if PLATFORM_LUMIN
 	// Only testing straight forward from a little below the headset... Lumin isn't a handheld, but it's nice to have this do something.
 	IXRTrackingSystem* XRTrackingSystem = Session->GetARSystem()->GetXRTrackingSystem();
 	TArray<int32> Devices;
@@ -397,7 +340,7 @@ void FLuminARFrame::ARLineTrace(FVector2D ScreenPosition, ELuminARLineTraceChann
 		FQuat HMDQuat;
 		FVector HMDPosition;
 		const bool Success = XRTrackingSystem->GetCurrentPose(HMDDeviceID, HMDQuat, HMDPosition);
-		FTransform TrackingToWorldTransform = XRTrackingSystem->GetTrackingToWorldTransform();
+		const FTransform TrackingToWorldTransform = XRTrackingSystem->GetTrackingToWorldTransform();
 		if (Success)
 		{
 			const FVector HMDWorldPosition = TrackingToWorldTransform.TransformPosition(HMDPosition);
@@ -409,12 +352,9 @@ void FLuminARFrame::ARLineTrace(FVector2D ScreenPosition, ELuminARLineTraceChann
 			ARLineTrace(Start, End, RequestedTraceChannels, OutHitResults);
 		}
 	}
-#endif
 }
 void FLuminARFrame::ARLineTrace(FVector Start, FVector End, ELuminARLineTraceChannel RequestedTraceChannels, TArray<FARTraceResult>& OutHitResults) const
 {
-#if PLATFORM_LUMIN
-
 	// Only testing vs planes now, but not the ground plane.
 	ELuminARLineTraceChannel AllPlaneTraceChannels = /*ELuminARLineTraceChannel::InfinitePlane |*/ ELuminARLineTraceChannel::PlaneUsingExtent | ELuminARLineTraceChannel::PlaneUsingBoundaryPolygon;
 	if (!(RequestedTraceChannels & AllPlaneTraceChannels))
@@ -534,7 +474,6 @@ void FLuminARFrame::ARLineTrace(FVector Start, FVector End, ELuminARLineTraceCha
 
 	// Sort closest to furthest
 	OutHitResults.Sort(FARTraceResult::FARTraceResultComparer());
-#endif
 }
 
 
@@ -542,7 +481,6 @@ FMatrix FLuminARFrame::GetProjectionMatrix() const
 {
 	FMatrix ProjectionMatrix;
 
-#if PLATFORM_LUMIN
 	if (Session == nullptr)
 	{
 		return ProjectionMatrix;
@@ -554,260 +492,170 @@ FMatrix FLuminARFrame::GetProjectionMatrix() const
 	ProjectionMatrix.M[2][2] = 0.0f;
 	ProjectionMatrix.M[2][3] = 1.0f;
 	ProjectionMatrix.M[3][2] = GNearClippingPlane;
-#endif
 	return ProjectionMatrix;
 }
 
 void FLuminARFrame::TransformDisplayUvCoords(const TArray<float>& UvCoords, TArray<float>& OutUvCoords) const
 {
-#if PLATFORM_LUMIN
 	OutUvCoords = UvCoords;
-#endif
 }
 
 
 FLuminARLightEstimate FLuminARFrame::GetLightEstimate() const
 {
-#if PLATFORM_LUMIN
+#if WITH_MLSDK
 	//TODO - Fill in light estimate.  See ml_lighting_tracking.h.
 	FLuminARLightEstimate LightEstimate;
 	return LightEstimate;
 #else
 	return FLuminARLightEstimate();
-#endif
+#endif // WITH_MLSDK
 }
 
 void FLuminARFrame::StartPlaneQuery()
 {
-#if PLATFORM_LUMIN
 	//if we haven't queried yet, start one!
-	if (PlaneQueryHandle == ML_INVALID_HANDLE)
+	if (!bPlanesQueryPending)
 	{
 		if (IMagicLeapPlugin::Get().IsMagicLeapHMDValid())
 		{
-			const FAppFramework& AppFramework = static_cast<FMagicLeapHMD*>(GEngine->XRSystem->GetHMDDevice())->GetAppFrameworkConst();
-			float WorldToMetersScale = AppFramework.GetWorldToMetersScale();
-			check(WorldToMetersScale != 0);
-
-			FTransform PoseInverse = UHeadMountedDisplayFunctionLibrary::GetTrackingToWorldTransform(nullptr).Inverse();
+			// TODO @njain : Since IMagicLeapPlanesModule::QueryBeginAsync() takes in search volume in world space, this is not needed. Remove usage!
+			const FTransform PoseInverse = UHeadMountedDisplayFunctionLibrary::GetTrackingToWorldTransform(nullptr).Inverse();
 			FPlane plane;
 
 			// Apply Lumin specific AR session config, if available.  Otherwise use default values.
-			MaxPlaneQueryResults = 200;
-			int32 MinPlaneArea = 25;
-			TArray<EPlaneQueryFlags> QueryFlags;
-			FVector SearchVolumeExtents(10000.0f, 10000.0f, 10000.0f);
+			FMagicLeapPlanesQuery Query;
+			Query.Flags = { EMagicLeapPlaneQueryFlags::Polygons };
+			Query.MaxResults = 200;
+			Query.MinHoleLength = 50.0f;
+			Query.MinPlaneArea = 400.0f;
+			Query.MinHoleLength = 50.0f;
+			Query.SearchVolumePosition = PoseInverse.GetTranslation();
+			Query.SearchVolumeOrientation = PoseInverse.GetRotation();
+			Query.SearchVolumeExtents = FVector(10000.0f, 10000.0f, 10000.0f);
 			bDiscardZeroExtentPlanes = false;
 
 			const UARSessionConfig& ARSessionConfig = Session->GetARSystem()->AccessSessionConfig();
 			const ULuminARSessionConfig* LuminARSessionConfig = Cast<ULuminARSessionConfig>(&ARSessionConfig);
 			if (LuminARSessionConfig != nullptr)
 			{
-				MaxPlaneQueryResults = LuminARSessionConfig->MaxPlaneQueryResults;
-				MinPlaneArea = LuminARSessionConfig->MinPlaneArea;
-				if (LuminARSessionConfig->ShouldDoHorizontalPlaneDetection())	{ QueryFlags.Add(EPlaneQueryFlags::Horizontal); }
-				if (LuminARSessionConfig->ShouldDoVerticalPlaneDetection())		{ QueryFlags.Add(EPlaneQueryFlags::Vertical); }
-				if (LuminARSessionConfig->bArbitraryOrientationPlaneDetection)	{ QueryFlags.Add(EPlaneQueryFlags::Arbitrary); }
-				SearchVolumeExtents = LuminARSessionConfig->PlaneSearchExtents;
-				for (EPlaneQueryFlags flag : LuminARSessionConfig->PlaneQueryFlags)
+				Query.MaxResults = LuminARSessionConfig->MaxPlaneQueryResults;
+				Query.MinPlaneArea = LuminARSessionConfig->MinPlaneArea;
+				if (LuminARSessionConfig->ShouldDoHorizontalPlaneDetection())	{ Query.Flags.Add(EMagicLeapPlaneQueryFlags::Horizontal); }
+				if (LuminARSessionConfig->ShouldDoVerticalPlaneDetection())		{ Query.Flags.Add(EMagicLeapPlaneQueryFlags::Vertical); }
+				if (LuminARSessionConfig->bArbitraryOrientationPlaneDetection)	{ Query.Flags.Add(EMagicLeapPlaneQueryFlags::Arbitrary); }
+				for (EMagicLeapPlaneQueryFlags flag : LuminARSessionConfig->PlaneQueryFlags)
 				{
-					QueryFlags.Add(flag);
+					Query.Flags.Add(flag);
 				}
+				Query.SearchVolumeExtents = LuminARSessionConfig->PlaneSearchExtents;
 				bDiscardZeroExtentPlanes = LuminARSessionConfig->bDiscardZeroExtentPlanes;
 			}
 			else
 			{
 				UE_LOG(LogLuminARAPI, Log, TEXT("LuminArSessionConfig not found, using defaults for lumin specific settings."));
-				QueryFlags.Add(EPlaneQueryFlags::Vertical);
-				QueryFlags.Add(EPlaneQueryFlags::Horizontal);
-			}
-
-			MLPlanesQuery Query;
-			Query.max_results = static_cast<uint32>(MaxPlaneQueryResults);
-			Query.flags = UnrealToMLPlanesQueryFlags(QueryFlags) | MLPlanesQueryFlag_Polygons;
-			Query.min_hole_length = 50.0f / WorldToMetersScale;  // ML docs say this value is deprecated, so presumably it does nothing now.
-			Query.min_plane_area = MinPlaneArea / (WorldToMetersScale * WorldToMetersScale);
-			Query.bounds_center = MagicLeap::ToMLVector(PoseInverse.GetTranslation(), WorldToMetersScale);
-			Query.bounds_rotation = MagicLeap::ToMLQuat(PoseInverse.GetRotation());
-			Query.bounds_extents = MagicLeap::ToMLVector(SearchVolumeExtents, WorldToMetersScale);
-
-			// MagicLeap::ToMLVector() causes the Z component to be negated.
-			// The bounds were thus invalid and resulted in everything being tracked. 
-			// This provides the content devs with an option to ignore the bounding volume at will.
-			{
-				Query.bounds_extents.x = FMath::Abs<float>(Query.bounds_extents.x);
-				Query.bounds_extents.y = FMath::Abs<float>(Query.bounds_extents.y);
-				Query.bounds_extents.z = FMath::Abs<float>(Query.bounds_extents.z);
+				Query.Flags.Add(EMagicLeapPlaneQueryFlags::Vertical);
+				Query.Flags.Add(EMagicLeapPlaneQueryFlags::Horizontal);
 			}
 
-			MLResult QueryResult = MLPlanesQueryBegin(PlaneTrackerHandle, &Query, &PlaneQueryHandle);
-			if (QueryResult != MLResult_Ok || !MLHandleIsValid(PlaneQueryHandle))
+			bPlanesQueryPending = IMagicLeapPlanesModule::Get().QueryBeginAsync(Query, ResultDelegate);
+			if (!bPlanesQueryPending)
 			{
 				UE_LOG(LogLuminARAPI, Error, TEXT("LuminARFrame could not request planes."));
 			}
 		}
 	}
-#endif //PLATFORM_LUMIN
 }
 
-void FLuminARFrame::ProcessPlaneQuery()
+void FLuminARFrame::ProcessPlaneQuery(const bool bSuccess, const TArray<FMagicLeapPlaneResult>& Planes, const TArray<FMagicLeapPlaneBoundaries>& Polygons)
 {
-#if PLATFORM_LUMIN
-	if (PlaneQueryHandle != ML_INVALID_HANDLE)
+	if (bSuccess)
 	{
-		FTransform PoseTransform = FTransform::Identity; 
+		const FTransform WorldToTrackingTransform = UHeadMountedDisplayFunctionLibrary::GetTrackingToWorldTransform(nullptr).Inverse();
+		PlaneResultsMap.Reset();
+		PlaneResultsMap.Reserve(Planes.Num());
 
-		uint32 OutNumResults = 0;
-		TArray<MLPlane> ResultMLPlanes;
-		ResultMLPlanes.AddDefaulted(MaxPlaneQueryResults);
-
-		MLPlaneBoundariesList PlaneBoundariesList;
-		MLPlaneBoundariesListInit(&PlaneBoundariesList);
-
-		MLResult PlaneQueryResult = MLPlanesQueryGetResultsWithBoundaries(PlaneTrackerHandle, PlaneQueryHandle, ResultMLPlanes.GetData(), &OutNumResults, &PlaneBoundariesList);
-		switch (PlaneQueryResult)
-		{
-		case MLResult_Pending:
-			// Intentionally skip. We'll continue to check until it has completed.
-			break;
-		case MLResult_UnspecifiedFailure:
-		{
-			UE_LOG(LogLuminARAPI, Error, TEXT("MLPlanesQueryGetResults MLResult_UnspecifiedFailure."));
-			PlaneQueryHandle = ML_INVALID_HANDLE;
-			LatestARPlaneQueryStatus = ELuminARPlaneQueryStatus::Fail;
-			break;
-		}
-		case MLResult_Ok:
+		for (const FMagicLeapPlaneResult& ResultUEPlane : Planes)
 		{
-			const FAppFramework& AppFramework = static_cast<FMagicLeapHMD*>(GEngine->XRSystem->GetHMDDevice())->GetAppFrameworkConst();
-			float WorldToMetersScale = AppFramework.GetWorldToMetersScale();
-
-			PlaneResultsMap.Reset();
-			PlaneResultsMap.Reserve(OutNumResults);
-
-			// Setup for boundaries, build a map of Handles to Boundaries
-			TMap<uint64, const MLPlaneBoundaries*> HandleToBoundariesMap;
+			if (ResultUEPlane.PlaneDimensions.X == 0.0f || ResultUEPlane.PlaneDimensions.Y == 0.0f)
 			{
-				const uint32 BoundaryCount = PlaneBoundariesList.plane_boundaries_count;
-				for (uint32 b = 0; b < BoundaryCount; ++b)
+				if (bDiscardZeroExtentPlanes)
 				{
-					const MLPlaneBoundaries& Boundaries = PlaneBoundariesList.plane_boundaries[b];
-					HandleToBoundariesMap.Add(Boundaries.id, &Boundaries);
-				}
-			}
-			const FRotator RotateToContentOrientation(-90.0f, 0.0f, 0.0f);
-			const FTransform RotateToContentTransform(RotateToContentOrientation);
-
-			for (uint32 i = 0; i < OutNumResults; ++i)
-			{
-				const MLPlane& ResultMLPlane = ResultMLPlanes[i];
-				FPlaneResult ResultUEPlane;
-
-				const uint64_t Mask = TNumericLimits<uint32>::Max();
-				ResultUEPlane.ID.A = ResultMLPlane.id & Mask;
-				ResultUEPlane.ID.B = ResultMLPlane.id >> 32;
-				ResultUEPlane.ID_64 = ResultMLPlane.id;
-
-				// Perception uses all coordinates in RUB so for them X axis is right and corresponds to the width of the plane.
-				// Unreal uses FRU, so the Y-axis is towards the right which makes the Y component of the vector the width.
-				ResultUEPlane.PlaneDimensions = FVector2D(ResultMLPlane.height * WorldToMetersScale, ResultMLPlane.width * WorldToMetersScale);
-				if (ResultUEPlane.PlaneDimensions.X == 0.0f || ResultUEPlane.PlaneDimensions.Y == 0.0f)
-				{
-					if (bDiscardZeroExtentPlanes)
+					if (PlaneResultsMap.Contains(ResultUEPlane.ID))
 					{
-						continue;
+						PlaneResultsMap.Remove(ResultUEPlane.ID);
 					}
+					continue;
 				}
+			}
+			FPlanesAndBoundaries& CachedResult = PlaneResultsMap.FindOrAdd(ResultUEPlane.ID);
+			// TODO: @njain We are only saving 1 boundary per ID :(
+			CachedResult.Plane = ResultUEPlane;
+			CachedResult.Plane.PlanePosition = WorldToTrackingTransform.TransformPosition(CachedResult.Plane.PlanePosition);
+			CachedResult.Plane.ContentOrientation = WorldToTrackingTransform.TransformRotation(CachedResult.Plane.ContentOrientation.Quaternion()).Rotator();
+			const float HalfWidth = CachedResult.Plane.PlaneDimensions.Y * 0.5f;
+			const float HalfHeight = CachedResult.Plane.PlaneDimensions.X * 0.5f;
+			CachedResult.PolygonVerticesLocalSpace = TArray<FVector>( 
+				{ FVector(HalfHeight, HalfWidth, 0)
+				, FVector(-HalfHeight, HalfWidth, 0)
+				, FVector(-HalfHeight, -HalfWidth, 0)
+				, FVector(HalfHeight, -HalfWidth, 0) } );
+		}
 
-				FTransform planeTransform = FTransform(MagicLeap::ToFQuat(ResultMLPlane.rotation), MagicLeap::ToFVector(ResultMLPlane.position, WorldToMetersScale), FVector(1.0f, 1.0f, 1.0f));
-				if (planeTransform.ContainsNaN())
+		// Setup for boundaries, build a map of Handles to Boundaries
+		for (const FMagicLeapPlaneBoundaries& AllBoundaries : Polygons)
+		{
+			if (AllBoundaries.Boundaries.Num() > 0)
+			{
+				FPlanesAndBoundaries* CachedResult = PlaneResultsMap.Find(AllBoundaries.ID);
+				if (CachedResult == nullptr)
 				{
 					continue;
 				}
-				if (!planeTransform.GetRotation().IsNormalized())
-				{
-					FQuat rotation = planeTransform.GetRotation();
-					rotation.Normalize();
-					planeTransform.SetRotation(rotation);
-				}
-
-				planeTransform.SetRotation(MagicLeap::ToUERotator(planeTransform.GetRotation()));
-				planeTransform.AddToTranslation(PoseTransform.GetLocation());
-				planeTransform.ConcatenateRotation(PoseTransform.Rotator().Quaternion());
-				ResultUEPlane.PlanePosition = planeTransform.GetLocation();
-				ResultUEPlane.PlaneOrientation = planeTransform.Rotator();
-				// The plane orientation has the forward axis (X) pointing in the direction of the plane's normal.
-				// We are rotating it by 90 degrees clock-wise about the right axis (Y) to get the up vector (Z) to point in the direction of the plane's normal.
-				// ResultPlane we are rotating the axis, the rotation is in the opposite direction of the object i.e. -90 degrees.
-				ResultUEPlane.ContentOrientation = UKismetMathLibrary::Conv_VectorToRotator(UKismetMathLibrary::RotateAngleAxis(UKismetMathLibrary::Conv_RotatorToVector(ResultUEPlane.PlaneOrientation), -90, UKismetMathLibrary::GetRightVector(ResultUEPlane.PlaneOrientation)));
-				MLToUnrealPlanesQueryFlags(ResultMLPlane.flags, ResultUEPlane.PlaneFlags);
-
-				// Boundaries
-				const MLPlaneBoundaries** const BoundariesPtr = HandleToBoundariesMap.Find(ResultMLPlane.id);
-				const MLPlaneBoundaries* const Boundaries = BoundariesPtr ? *BoundariesPtr : nullptr;
-				if (Boundaries && Boundaries->boundaries_count > 0)
+				const FTransform TrackingToLocalTransform = FTransform(CachedResult->Plane.ContentOrientation, CachedResult->Plane.PlanePosition).Inverse();
+				const FTransform WorldToLocalTransform = WorldToTrackingTransform * TrackingToLocalTransform;
+				// TODO: @njain We are only saving 1 boundary per ID :(
+				CachedResult->PolygonVerticesLocalSpace = AllBoundaries.Boundaries[0].Polygon.Vertices;
+				for (FVector& Vertex : CachedResult->PolygonVerticesLocalSpace)
 				{
-					FTransform PlaneTransformInverse = planeTransform.Inverse();
-					FTransform BoundaryVertTransform = PlaneTransformInverse * RotateToContentTransform;
-
-					// Seems like there is really only one boundary polygon...
-					const int BoundariesCount = Boundaries->boundaries_count;
-					for (int b = 0; b < BoundariesCount; ++b)
-					{
-						const MLPlaneBoundary& Boundary = Boundaries->boundaries[b];
-						const MLPolygon* Polygon = Boundary.polygon;
-						check(Polygon);
-						const uint32 VertCount = Polygon->vertices_count;
-						for (uint32 v = 0; v < VertCount; ++v)
-						{
-							MLVec3f& Vert = Polygon->vertices[v];
-							FVector LocalVert = BoundaryVertTransform.TransformPosition(MagicLeap::ToFVector(Vert, WorldToMetersScale));
-							ResultUEPlane.BoundaryPolygon.Add(LocalVert);
-						}
-					}
+					Vertex = WorldToLocalTransform.TransformPosition(Vertex);
 				}
-
-				PlaneResultsMap.Add(ResultMLPlane.id, ResultUEPlane);
 			}
+		}
 
-			PlaneQueryHandle = ML_INVALID_HANDLE;
-
-			// Mark planes that previously existed, but no longer do StoppedTracking.
-			const auto & TrackableHandleMap = Session->GetUObjectManager()->TrackableHandleMap;
-			for (auto TrackablePair : TrackableHandleMap)
+		// Mark planes that previously existed, but no longer do StoppedTracking.
+		const auto & TrackableHandleMap = Session->GetUObjectManager()->TrackableHandleMap;
+		for (auto TrackablePair : TrackableHandleMap)
+		{
+			if (!PlaneResultsMap.Contains(TrackablePair.Key))
 			{
-				uint64 Handle = TrackablePair.Key;
-				if (!PlaneResultsMap.Contains(Handle))
+				UARTrackedGeometry* ARTrackedGeometry = TrackablePair.Value.Get();
+				if (ARTrackedGeometry)
 				{
-					UARTrackedGeometry* ARTrackedGeometry = TrackablePair.Value.Get();
-					if (ARTrackedGeometry)
-					{
-						ARTrackedGeometry->SetTrackingState(EARTrackingState::StoppedTracking);
-					}
+					ARTrackedGeometry->SetTrackingState(EARTrackingState::StoppedTracking);
 				}
 			}
+		}
 
-			for (auto PlaneResultPair : PlaneResultsMap)
-			{
-				UARTrackedGeometry* ARTrackedGeometry = Session->GetUObjectManager()->GetTrackableFromHandle<UARPlaneGeometry>(PlaneResultPair.Key, Session);
+		for (auto PlaneResultPair : PlaneResultsMap)
+		{
+			UARTrackedGeometry* ARTrackedGeometry = Session->GetUObjectManager()->GetTrackableFromHandle<UARPlaneGeometry>(PlaneResultPair.Key, Session);
 
-				if (ARTrackedGeometry && ARTrackedGeometry->GetTrackingState() != EARTrackingState::StoppedTracking)
-				{
-					ARTrackedGeometry->SetTrackingState(EARTrackingState::Tracking);
-					FLuminARTrackableResource* TrackableResource = reinterpret_cast<FLuminARTrackableResource*>(ARTrackedGeometry->GetNativeResource());
-					TrackableResource->UpdateGeometryData(Session);
-				}
+			if (ARTrackedGeometry && ARTrackedGeometry->GetTrackingState() != EARTrackingState::StoppedTracking)
+			{
+				ARTrackedGeometry->SetTrackingState(EARTrackingState::Tracking);
+				FLuminARTrackableResource* TrackableResource = static_cast<FLuminARTrackableResource*>(ARTrackedGeometry->GetNativeResource());
+				TrackableResource->UpdateGeometryData(Session);
 			}
-			LatestARPlaneQueryStatus = ELuminARPlaneQueryStatus::Success;
-			break;
-		}
-		default:
-			UE_LOG(LogLuminARAPI, Warning, TEXT("Unexpected return code from MLPlanesQueryGetResults: %d"), PlaneQueryResult);
-			LatestARPlaneQueryStatus = ELuminARPlaneQueryStatus::Fail;
 		}
-		MLPlanesReleaseBoundariesList(PlaneTrackerHandle, &PlaneBoundariesList);
+		LatestARPlaneQueryStatus = ELuminARPlaneQueryStatus::Success;
 	}
-#endif //PLATFORM_LUMIN
+	else
+	{
+		LatestARPlaneQueryStatus = ELuminARPlaneQueryStatus::Fail;
+	}
+
+	bPlanesQueryPending = false;
 }
 
 
@@ -827,11 +675,10 @@ TSharedPtr<FLuminARSession> FLuminARSession::CreateLuminARSession()
 /************************************************/
 /*       ULuminARTrackableResource         */
 /************************************************/
-#if PLATFORM_LUMIN
 
 EARTrackingState FLuminARTrackableResource::GetTrackingState()
 {
-	if (MLHandleIsValid(TrackableHandle))
+	if (MagicLeap::FGuidIsValidHandle(TrackableHandle))
 	{
 		check(TrackedGeometry);
 		return TrackedGeometry->GetTrackingState();
@@ -847,7 +694,9 @@ void FLuminARTrackableResource::UpdateGeometryData(FLuminARSession* InSession)
 
 void FLuminARTrackableResource::ResetNativeHandle(LuminArTrackable* InTrackableHandle)
 {
-	TrackableHandle = ML_INVALID_HANDLE;
+#if WITH_MLSDK
+	TrackableHandle = MagicLeap::MLHandleToFGuid(ML_INVALID_HANDLE);
+#endif // WITH_MLSDK
 	if (InTrackableHandle != nullptr)
 	{
 		TrackableHandle = InTrackableHandle->Handle;
@@ -874,49 +723,52 @@ void FLuminARTrackedPlaneResource::UpdateGeometryData(FLuminARSession* InSession
 	}
 
 	// PlaneResult is in unreal tracking space, so already scaled and axis corrected
-	const FPlaneResult* PlaneResult = Frame->GetPlaneResult(TrackableHandle);
-	check(PlaneResult);
+	const FPlanesAndBoundaries* Results = Frame->GetPlaneResult(TrackableHandle);
+	check(Results);
+	const FMagicLeapPlaneResult* PlaneResult = &Results->Plane;
 
-	FTransform LocalToTrackingTransform(PlaneResult->ContentOrientation, PlaneResult->PlanePosition);
+	// The ContentOrientation and PlanePosition have been transformed back to Tracking space in ProcessPlanesQuery()
+	const FTransform LocalToTrackingTransform(PlaneResult->ContentOrientation, PlaneResult->PlanePosition);
 	FVector Extent(PlaneResult->PlaneDimensions.X * 0.5f, PlaneResult->PlaneDimensions.Y * 0.5f, 0);  // Extent is half the width and height
 
 	uint32 FrameNum = InSession->GetFrameNum();
 	int64 TimeStamp = InSession->GetLatestFrame()->GetCameraTimestamp();
 
 	UARPlaneGeometry* SubsumedByPlane = nullptr; // ARCore only
-	PlaneGeometry->UpdateTrackedGeometry(InSession->GetARSystem(), FrameNum, static_cast<double>(TimeStamp), LocalToTrackingTransform, InSession->GetARSystem()->GetAlignmentTransform(), FVector::ZeroVector, Extent, PlaneResult->BoundaryPolygon, SubsumedByPlane);
+	// TODO: @njain confirm the space all these coordinates need to be in.
+	PlaneGeometry->UpdateTrackedGeometry(InSession->GetARSystem(), FrameNum, static_cast<double>(TimeStamp), LocalToTrackingTransform, InSession->GetARSystem()->GetAlignmentTransform(), FVector::ZeroVector, Extent, Results->PolygonVerticesLocalSpace, SubsumedByPlane);
 	PlaneGeometry->SetDebugName(FName(TEXT("LuminARPlane")));
 	// Inspect the plane flags and set the scene understanding data based upon them
-	for (EPlaneQueryFlags Flag : PlaneResult->PlaneFlags)
+	for (EMagicLeapPlaneQueryFlags Flag : PlaneResult->PlaneFlags)
 	{
 		switch (Flag)
 		{
-			case EPlaneQueryFlags::Vertical:
+			case EMagicLeapPlaneQueryFlags::Vertical:
 			{
 				PlaneGeometry->SetOrientation(EARPlaneOrientation::Vertical);
 				break;
 			}
-			case EPlaneQueryFlags::Horizontal:
+			case EMagicLeapPlaneQueryFlags::Horizontal:
 			{
 				PlaneGeometry->SetOrientation(EARPlaneOrientation::Horizontal);
 				break;
 			}
-			case EPlaneQueryFlags::Arbitrary:
+			case EMagicLeapPlaneQueryFlags::Arbitrary:
 			{
 				PlaneGeometry->SetOrientation(EARPlaneOrientation::Diagonal);
 				break;
 			}
-			case EPlaneQueryFlags::Ceiling:
+			case EMagicLeapPlaneQueryFlags::Ceiling:
 			{
 				PlaneGeometry->SetObjectClassification(EARObjectClassification::Ceiling);
 				break;
 			}
-			case EPlaneQueryFlags::Floor:
+			case EMagicLeapPlaneQueryFlags::Floor:
 			{
 				PlaneGeometry->SetObjectClassification(EARObjectClassification::Floor);
 				break;
 			}
-			case EPlaneQueryFlags::Wall:
+			case EMagicLeapPlaneQueryFlags::Wall:
 			{
 				PlaneGeometry->SetObjectClassification(EARObjectClassification::Wall);
 				break;
@@ -924,24 +776,21 @@ void FLuminARTrackedPlaneResource::UpdateGeometryData(FLuminARSession* InSession
 		}
 	}
 }
-#endif
 
 
-#if PLATFORM_LUMIN
-void ULuminARUObjectManager::DumpTrackableHandleMap(const MLHandle SessionHandle)
+void ULuminARUObjectManager::DumpTrackableHandleMap(const FGuid& SessionHandle)
 {
 	UE_LOG(LogLuminARAPI, Log, TEXT("ULuminARUObjectManager::DumpTrackableHandleMap"));
 	for (auto KeyValuePair : TrackableHandleMap)
 	{
-		uint64 TrackableHandle = KeyValuePair.Key;
 		TWeakObjectPtr<UARTrackedGeometry> TrackedGeometry = KeyValuePair.Value;
-		UE_LOG(LogLuminARAPI, Log, TEXT("  Trackable Handle %llu"), TrackableHandle);
+		UE_LOG(LogLuminARAPI, Log, TEXT("  Trackable Handle %s"), *KeyValuePair.Key.ToString());
 		if (TrackedGeometry.IsValid())
 		{
 			UARTrackedGeometry* TrackedGeometryObj = TrackedGeometry.Get();
-			FLuminARTrackableResource* NativeResource = reinterpret_cast<FLuminARTrackableResource*>(TrackedGeometryObj->GetNativeResource());
-			UE_LOG(LogLuminARAPI, Log, TEXT("  TrackedGeometry - NativeResource:%p, type: %s, tracking state: %d"),
-				NativeResource->GetNativeHandle(), *TrackedGeometryObj->GetClass()->GetFName().ToString(), (int)TrackedGeometryObj->GetTrackingState());
+			FLuminARTrackableResource* NativeResource = static_cast<FLuminARTrackableResource*>(TrackedGeometryObj->GetNativeResource());
+			UE_LOG(LogLuminARAPI, Log, TEXT("  TrackedGeometry - NativeResource:%s, type: %s, tracking state: %d"),
+				*NativeResource->GetNativeHandle().ToString(), *TrackedGeometryObj->GetClass()->GetFName().ToString(), (int)TrackedGeometryObj->GetTrackingState());
 		}
 		else
 		{
@@ -949,12 +798,11 @@ void ULuminARUObjectManager::DumpTrackableHandleMap(const MLHandle SessionHandle
 		}
 	}
 }
-#endif
 
 ELuminARAPIStatus FLuminARSession::AcquireCameraImage(ULuminARCameraImage *&OutCameraImage)
 {
 	ELuminARAPIStatus ApiStatus = ELuminARAPIStatus::AR_SUCCESS;
-#if PLATFORM_LUMIN
+#if WITH_MLSDK
 	if (LatestFrame == nullptr)
 	{
 		return ELuminARAPIStatus::AR_ERROR_FATAL;
@@ -971,19 +819,15 @@ ELuminARAPIStatus FLuminARSession::AcquireCameraImage(ULuminARCameraImage *&OutC
 	{
 		UE_LOG(LogLuminARAPI, Error, TEXT("AcquireCameraImage failed!"));
 	}
-#endif
+#endif // WITH_MLSDK
 
 	return ApiStatus;
 }
 
 void* FLuminARSession::GetLatestFrameRawPointer()
 {
-#if PLATFORM_LUMIN
+#if WITH_MLSDK
 	//TODO - do we need this?
-#endif
+#endif // WITH_MLSDK
 	return nullptr;
 }
-
-
-
-
diff --git a/Engine/Plugins/Lumin/MagicLeap/Source/MagicLeapAR/Private/LuminARAPI.h b/Engine/Plugins/Lumin/MagicLeap/Source/MagicLeapAR/Private/LuminARAPI.h
--- a/Engine/Plugins/Lumin/MagicLeap/Source/MagicLeapAR/Private/LuminARAPI.h
+++ b/Engine/Plugins/Lumin/MagicLeap/Source/MagicLeapAR/Private/LuminARAPI.h
@@ -7,7 +7,8 @@
 #include "LuminARTypes.h"
 #include "LuminARSessionConfig.h"
 #include "ARSessionConfig.h"
-#include "PlanesComponent.h"
+#include "MagicLeapPlanesTypes.h"
+#include "MagicLeapHandle.h"
 
 #include "LuminARAPI.generated.h"
 
@@ -35,7 +36,6 @@ enum class ELuminARPlaneQueryStatus : int
 	Fail
 };
 
-#if PLATFORM_LUMIN
 static ArTrackableType GetTrackableType(UClass* ClassType)
 {
 	if (ClassType == UARTrackedGeometry::StaticClass())
@@ -51,7 +51,6 @@ static ArTrackableType GetTrackableType(UClass* ClassType)
 		return ArTrackableType::LUMIN_AR_TRACKABLE_NOT_VALID;
 	}
 }
-#endif
 
 class FLuminARFrame;
 class FLuminARSession;
@@ -66,15 +65,13 @@ public:
 	UPROPERTY()
 	TArray<UARPin*> AllAnchors;
 
-#if PLATFORM_LUMIN
-	TMap<uint64, TSharedPtr<LuminArAnchor>> HandleToLuminAnchorMap;
-	TMap<uint64, UARPin*> HandleToAnchorMap;
-	TMap<uint64, TWeakObjectPtr<UARTrackedGeometry>> TrackableHandleMap;
+	TMap<FGuid, TSharedPtr<LuminArAnchor>> HandleToLuminAnchorMap;
+	TMap<FGuid, UARPin*> HandleToAnchorMap;
+	TMap<FGuid, TWeakObjectPtr<UARTrackedGeometry>> TrackableHandleMap;
 
-	template< class T > T* GetTrackableFromHandle(MLHandle TrackableHandle, FLuminARSession* Session);
+	template< class T > T* GetTrackableFromHandle(const FGuid& TrackableHandle, FLuminARSession* Session);
 
-	void DumpTrackableHandleMap(const MLHandle SessionHandle);
-#endif
+	void DumpTrackableHandleMap(const FGuid& SessionHandle);
 };
 
 
@@ -92,9 +89,6 @@ public:
 	float GetWorldToMeterScale();
 	void SetARSystem(TSharedRef<FARSupportInterface , ESPMode::ThreadSafe> InArSystem) { ARSystem = InArSystem; }
 	TSharedRef<FARSupportInterface , ESPMode::ThreadSafe> GetARSystem() { return ARSystem.ToSharedRef(); }
-#if PLATFORM_LUMIN
-	MLHandle GetPlaneTrackerHandle();
-#endif
 
 	ELuminARAPIStatus Resume();
 	ELuminARAPIStatus Pause();
@@ -123,10 +117,13 @@ private:
 	uint32 FrameNumber;
 
 	TSharedPtr<FARSupportInterface , ESPMode::ThreadSafe> ARSystem;
+};
 
-#if PLATFORM_LUMIN
-	MLHandle PlaneTrackerHandle = ML_INVALID_HANDLE;
-#endif
+struct FPlanesAndBoundaries
+{
+public:
+	FMagicLeapPlaneResult Plane;
+	TArray<FVector> PolygonVerticesLocalSpace;
 };
 
 class FLuminARFrame
@@ -147,7 +144,7 @@ public:
 
 	void GetUpdatedAnchors(TArray<UARPin*>& OutUpdatedAnchors) const;
 
-	const FPlaneResult* GetPlaneResult(int64 Handle) const { return PlaneResultsMap.Find(Handle); }
+	const FPlanesAndBoundaries* GetPlaneResult(const FGuid& Handle) const { return PlaneResultsMap.Find(Handle); }
 
 	void ARLineTrace(FVector2D ScreenPosition, ELuminARLineTraceChannel RequestedTraceChannels, TArray<FARTraceResult>& OutHitResults) const;
 	void ARLineTrace(FVector Start, FVector End, ELuminARLineTraceChannel RequestedTraceChannels, TArray<FARTraceResult>& OutHitResults) const;
@@ -160,7 +157,7 @@ public:
 private:
 
 	void StartPlaneQuery();
-	void ProcessPlaneQuery();
+	void ProcessPlaneQuery(const bool bSuccess, const TArray<FMagicLeapPlaneResult>& Planes, const TArray<FMagicLeapPlaneBoundaries>& Polygons);
 
 
 	FLuminARSession* Session;
@@ -170,15 +167,14 @@ private:
 	ELuminARPlaneQueryStatus LatestARPlaneQueryStatus;
 
 	TArray<UARPin*> UpdatedAnchors;
-	TMap<uint64, FPlaneResult> PlaneResultsMap;
+	TMap<FGuid, FPlanesAndBoundaries> PlaneResultsMap;
 
 	int32 MaxPlaneQueryResults = 0;
 	bool bDiscardZeroExtentPlanes = false;
 
-#if PLATFORM_LUMIN
-	MLHandle PlaneTrackerHandle = ML_INVALID_HANDLE;
-	MLHandle PlaneQueryHandle = ML_INVALID_HANDLE;
-#endif
+	bool bPlanesQueryPending;
+
+	FMagicLeapPlanesResultStaticDelegate ResultDelegate;
 };
 
 class FLuminARTrackableResource : public IARRef
@@ -189,18 +185,19 @@ public:
 
 	virtual void RemoveRef() override
 	{
-#if PLATFORM_LUMIN
-		TrackableHandle = ML_INVALID_HANDLE;
-#endif
+#if WITH_MLSDK
+		TrackableHandle = MagicLeap::MLHandleToFGuid(ML_INVALID_HANDLE);
+#endif // WITH_MLSDK
 	}
 
-#if PLATFORM_LUMIN
 public:
-	FLuminARTrackableResource(MLHandle InTrackableHandle, UARTrackedGeometry* InTrackedGeometry)
+	FLuminARTrackableResource(const FGuid& InTrackableHandle, UARTrackedGeometry* InTrackedGeometry)
 		: TrackableHandle(InTrackableHandle)
 		, TrackedGeometry(InTrackedGeometry)
 	{
-		ensure(TrackableHandle != ML_INVALID_HANDLE);
+#if WITH_MLSDK
+		ensure(MagicLeap::FGuidIsValidHandle(TrackableHandle));
+#endif // WITH_MLSDK
 	}
 
 	virtual ~FLuminARTrackableResource()
@@ -211,36 +208,29 @@ public:
 
 	virtual void UpdateGeometryData(FLuminARSession* InSession);
 
-	MLHandle GetNativeHandle() { return TrackableHandle; }
+	const FGuid& GetNativeHandle() { return TrackableHandle; }
 
 	void ResetNativeHandle(LuminArTrackable* InTrackableHandle);
 
 protected:
-	MLHandle TrackableHandle;
+	FGuid TrackableHandle;
 	UARTrackedGeometry* TrackedGeometry;
-#endif
 };
 
 class FLuminARTrackedPlaneResource : public FLuminARTrackableResource
 {
 public:
-#if PLATFORM_LUMIN
-	FLuminARTrackedPlaneResource(MLHandle InTrackableHandle, UARTrackedGeometry* InTrackedGeometry)
+	FLuminARTrackedPlaneResource(const FGuid& InTrackableHandle, UARTrackedGeometry* InTrackedGeometry)
 		: FLuminARTrackableResource(InTrackableHandle, InTrackedGeometry)
 	{
-		ensure(TrackableHandle != ML_INVALID_HANDLE);
 	}
 
 	void UpdateGeometryData(FLuminARSession* InSession) override;
-
-	ArPlane* GetPlaneHandle() { return reinterpret_cast<ArPlane*>(TrackableHandle); }
-#endif
 };
 
-#if PLATFORM_LUMIN
 // Template function definition
 template< class T >
-T* ULuminARUObjectManager::GetTrackableFromHandle(MLHandle TrackableHandle, FLuminARSession* Session)
+T* ULuminARUObjectManager::GetTrackableFromHandle(const FGuid& TrackableHandle, FLuminARSession* Session)
 {
 	if (!TrackableHandleMap.Contains(TrackableHandle)
 		|| !TrackableHandleMap[TrackableHandle].IsValid()
@@ -253,7 +243,7 @@ T* ULuminARUObjectManager::GetTrackableFromHandle(MLHandle TrackableHandle, FLum
 		check(Session);
 		const FLuminARFrame* Frame = Session->GetLatestFrame();
 		check(Frame);
-		const FPlaneResult* PlaneResult = Frame->GetPlaneResult(TrackableHandle);
+		const FPlanesAndBoundaries* PlaneResult = Frame->GetPlaneResult(TrackableHandle);
 		if (PlaneResult)
 		{
 			TrackableType = ArTrackableType::LUMIN_AR_TRACKABLE_PLANE;
@@ -278,7 +268,7 @@ T* ULuminARUObjectManager::GetTrackableFromHandle(MLHandle TrackableHandle, FLum
 		NewTrackableObject->InitializeNativeResource(NativeResource);
 		NativeResource = nullptr;
 
-		FLuminARTrackableResource* TrackableResource = reinterpret_cast<FLuminARTrackableResource*>(NewTrackableObject->GetNativeResource());
+		FLuminARTrackableResource* TrackableResource = static_cast<FLuminARTrackableResource*>(NewTrackableObject->GetNativeResource());
 
 		// no, we always do this in ProcessPlaneQuery
 		// Update the tracked geometry data using the native resource
@@ -292,13 +282,11 @@ T* ULuminARUObjectManager::GetTrackableFromHandle(MLHandle TrackableHandle, FLum
 	//checkf(Result, TEXT("ULuminARUObjectManager failed to get a valid trackable %p from the map."), TrackableHandle);
 	return Result;
 }
-#endif
 
 template< class T >
 void FLuminARSession::GetAllTrackables(TArray<T*>& OutARCoreTrackableList)
 {
 	OutARCoreTrackableList.Empty();
-#if PLATFORM_LUMIN
 
 	ArTrackableType TrackableType = GetTrackableType(T::StaticClass());
 	if (TrackableType == ArTrackableType::LUMIN_AR_TRACKABLE_NOT_VALID)
@@ -321,5 +309,4 @@ void FLuminARSession::GetAllTrackables(TArray<T*>& OutARCoreTrackableList)
 			}
 		}	
 	}
-#endif
 }
diff --git a/Engine/Plugins/Lumin/MagicLeap/Source/MagicLeapAR/Private/LuminARDevice.cpp b/Engine/Plugins/Lumin/MagicLeap/Source/MagicLeapAR/Private/LuminARDevice.cpp
--- a/Engine/Plugins/Lumin/MagicLeap/Source/MagicLeapAR/Private/LuminARDevice.cpp
+++ b/Engine/Plugins/Lumin/MagicLeap/Source/MagicLeapAR/Private/LuminARDevice.cpp
@@ -253,7 +253,7 @@ void FLuminARDevice::StartSession()
 
 void* FLuminARDevice::GetARSessionRawPointer()
 {
-#if PLATFORM_LUMIN
+#if WITH_MLSDK
 	check(false); // never used and unsupported
 	return nullptr;
 #endif
@@ -262,7 +262,7 @@ void* FLuminARDevice::GetARSessionRawPointer()
 
 void* FLuminARDevice::GetGameThreadARFrameRawPointer()
 {
-#if PLATFORM_LUMIN
+#if WITH_MLSDK
 	return reinterpret_cast<void*>(LuminARSession->GetLatestFrameRawPointer());
 #endif
 	return nullptr;
diff --git a/Engine/Plugins/Lumin/MagicLeap/Source/MagicLeapAR/Private/LuminARModule.cpp b/Engine/Plugins/Lumin/MagicLeap/Source/MagicLeapAR/Private/LuminARModule.cpp
--- a/Engine/Plugins/Lumin/MagicLeap/Source/MagicLeapAR/Private/LuminARModule.cpp
+++ b/Engine/Plugins/Lumin/MagicLeap/Source/MagicLeapAR/Private/LuminARModule.cpp
@@ -8,6 +8,7 @@
 #include "LuminARTrackingSystem.h"
 #include "LuminARDevice.h"
 #include "Templates/SharedPointer.h"
+#include "IMagicLeapPlugin.h"
 
 
 #define LOCTEXT_NAMESPACE "LuminAR"
@@ -49,9 +50,16 @@ void FLuminARModule::ShutdownModule()
 //create for mutual connection (regardless of construction order)
 TSharedPtr<IARSystemSupport, ESPMode::ThreadSafe> FLuminARModule::CreateARImplementation()
 {
-#if PLATFORM_LUMIN
-	LuminARImplementation = MakeShareable(new FLuminARImplementation());
-#endif
+#if WITH_MLSDK
+#if !PLATFORM_LUMIN
+	bool bIsVDZIEnabled = false;
+	GConfig->GetBool(TEXT("/Script/MagicLeap.MagicLeapSettings"), TEXT("bEnableZI"), bIsVDZIEnabled, GEngineIni);
+	if (bIsVDZIEnabled)
+#endif // !PLATFORM_LUMIN
+	{
+		LuminARImplementation = MakeShareable(new FLuminARImplementation());
+	}
+#endif //WITH_MLSDK
 	return LuminARImplementation;
 }
 
diff --git a/Engine/Plugins/Lumin/MagicLeap/Source/MagicLeapAR/Private/LuminARTrackingSystem.cpp b/Engine/Plugins/Lumin/MagicLeap/Source/MagicLeapAR/Private/LuminARTrackingSystem.cpp
--- a/Engine/Plugins/Lumin/MagicLeap/Source/MagicLeapAR/Private/LuminARTrackingSystem.cpp
+++ b/Engine/Plugins/Lumin/MagicLeap/Source/MagicLeapAR/Private/LuminARTrackingSystem.cpp
@@ -31,20 +31,12 @@ FLuminARImplementation::~FLuminARImplementation()
 
 void* FLuminARImplementation::GetARSessionRawPointer()
 {
-#if PLATFORM_LUMIN
 	return static_cast<void*>(FLuminARDevice::GetInstance()->GetARSessionRawPointer());
-#endif
-	ensureAlwaysMsgf(false, TEXT("FLuminARImplementation::GetARSessionRawPointer is unimplemented on current platform."));
-	return nullptr;
 }
 
 void* FLuminARImplementation::GetGameThreadARFrameRawPointer()
 {
-#if PLATFORM_LUMIN
 	return static_cast<void*>(FLuminARDevice::GetInstance()->GetGameThreadARFrameRawPointer());
-#endif
-	ensureAlwaysMsgf(false, TEXT("FLuminARImplementation::GetARSessionRawPointer is unimplemented on current platform."));
-	return nullptr;
 }
 
 void FLuminARImplementation::OnARSystemInitialized()
diff --git a/Engine/Plugins/Lumin/MagicLeap/Source/MagicLeapAR/Public/LuminARModule.h b/Engine/Plugins/Lumin/MagicLeap/Source/MagicLeapAR/Public/LuminARModule.h
--- a/Engine/Plugins/Lumin/MagicLeap/Source/MagicLeapAR/Public/LuminARModule.h
+++ b/Engine/Plugins/Lumin/MagicLeap/Source/MagicLeapAR/Public/LuminARModule.h
@@ -3,4 +3,3 @@
 #pragma once
 
 #include "Modules/ModuleManager.h"
-#include "MagicLeapHMD.h"
diff --git a/Engine/Plugins/Lumin/MagicLeap/Source/MagicLeapAR/Public/LuminARSessionConfig.h b/Engine/Plugins/Lumin/MagicLeap/Source/MagicLeapAR/Public/LuminARSessionConfig.h
--- a/Engine/Plugins/Lumin/MagicLeap/Source/MagicLeapAR/Public/LuminARSessionConfig.h
+++ b/Engine/Plugins/Lumin/MagicLeap/Source/MagicLeapAR/Public/LuminARSessionConfig.h
@@ -5,8 +5,8 @@
 #include "UObject/Object.h"
 #include "UObject/ObjectMacros.h"
 #include "LuminARTypes.h"
-#include "PlanesComponent.h"
 #include "ARSessionConfig.h"
+#include "MagicLeapPlanesTypes.h"
 
 #include "LuminARSessionConfig.generated.h"
 
@@ -44,7 +44,7 @@ public:
 
 	/** Additional Flags to apply to the plane queries. Note: the plane orientation detection settings also cause flags to be set.  It is ok to duplicate those here.*/
 	UPROPERTY(EditAnywhere, Category = "Lumin AR Settings")
-	TArray<EPlaneQueryFlags> PlaneQueryFlags;
+	TArray<EMagicLeapPlaneQueryFlags> PlaneQueryFlags;
 
 	/** If true discard any 'plane' objects that come through with zero extents and only polygon edge data.*/
 	UPROPERTY(EditAnywhere, Category = "Lumin AR Settings")
diff --git a/Engine/Plugins/Lumin/MagicLeap/Source/MagicLeapAR/Public/LuminARTypes.h b/Engine/Plugins/Lumin/MagicLeap/Source/MagicLeapAR/Public/LuminARTypes.h
--- a/Engine/Plugins/Lumin/MagicLeap/Source/MagicLeapAR/Public/LuminARTypes.h
+++ b/Engine/Plugins/Lumin/MagicLeap/Source/MagicLeapAR/Public/LuminARTypes.h
@@ -3,9 +3,8 @@
 
 #include "CoreMinimal.h"
 
-#if WITH_MLSDK
-#include <ml_api.h>
-#endif //WITH_MLSDK
+#include "Lumin/CAPIShims/LuminAPI.h"
+#include "MagicLeapHandle.h"
 
 #include "ARTypes.h"
 #include "ARTrackable.h"
@@ -18,8 +17,6 @@
 /// @defgroup LuminARBase
 /// The base module for LuminAR plugin
 
-#if PLATFORM_LUMIN
-
 typedef enum
 {
 	LUMIN_AR_TRACKABLE_NOT_VALID,
@@ -38,26 +35,39 @@ struct ArPose
 struct LuminArTrackable
 {
 	LuminArTrackable(ArPose InPose)
-		: Handle(ML_INVALID_HANDLE)
+#if WITH_MLSDK
+		: Handle(MagicLeap::MLHandleToFGuid(ML_INVALID_HANDLE))
+#endif // WITH_MLSDK
 	{
 	}
 
-	MLHandle Handle;
+	// Use FGuid since a LuminArTrackable can either be an MLHandle or an MLCFUID (currently only for PCFs)
+	FGuid Handle;
 };
 
 struct LuminArAnchor : public LuminArTrackable
 {
-	LuminArAnchor(ArPose InPose, MLHandle InParentTrackable)
+	LuminArAnchor(ArPose InPose)
+		: LuminArTrackable(InPose)
+#if WITH_MLSDK
+		, ParentTrackable(MagicLeap::MLHandleToFGuid(ML_INVALID_HANDLE))
+#endif // WITH_MLSDK
+	{}
+
+	LuminArAnchor(ArPose InPose, const FGuid& InParentTrackable)
 		: LuminArTrackable(InPose)
 		, ParentTrackable(InParentTrackable)
 	{}
 
 	void Detach()
 	{
-		ParentTrackable = ML_INVALID_HANDLE;
+#if WITH_MLSDK
+		ParentTrackable = MagicLeap::MLHandleToFGuid(ML_INVALID_HANDLE);
+#endif // WITH_MLSDK
 	}
 
-	MLHandle ParentTrackable;
+	// Use FGuid since a LuminArTrackable can either be an MLHandle or an MLCFUID (currently only for PCFs)
+	FGuid ParentTrackable;
 };
 
 typedef struct MLPlane ArPlane;
@@ -68,8 +78,6 @@ struct ArImage
 {
 };
 
-#endif
-
 UENUM(BlueprintType)
 enum class ELuminARAvailability : uint8
 {
diff --git a/Engine/Plugins/Lumin/MagicLeap/Source/MagicLeapARPin/MagicLeapARPin.Build.cs b/Engine/Plugins/Lumin/MagicLeap/Source/MagicLeapARPin/MagicLeapARPin.Build.cs
new file mode 100644
--- /dev/null
+++ b/Engine/Plugins/Lumin/MagicLeap/Source/MagicLeapARPin/MagicLeapARPin.Build.cs
@@ -0,0 +1,32 @@
+// Copyright 1998-2019 Epic Games, Inc. All Rights Reserved.
+
+using System.IO;
+
+namespace UnrealBuildTool.Rules
+{
+	public class MagicLeapARPin : ModuleRules
+	{
+		public MagicLeapARPin( ReadOnlyTargetRules Target ) : base(Target)
+		{
+			PublicIncludePaths.Add(Path.Combine(ModuleDirectory, "Public"));
+			PublicDependencyModuleNames.AddRange(
+				new string[]
+				{
+					"Core",
+					"CoreUObject",
+					"Engine"
+				}
+			);
+
+			PrivateDependencyModuleNames.AddRange(
+				new string[]
+				{
+					"LuminRuntimeSettings",
+					"MLSDK",
+					"MagicLeap",
+					"HeadMountedDisplay",
+				}
+			);
+		}
+	}
+}
diff --git a/Engine/Plugins/Lumin/MagicLeap/Source/MagicLeapARPin/Private/MagicLeapARPinComponent.cpp b/Engine/Plugins/Lumin/MagicLeap/Source/MagicLeapARPin/Private/MagicLeapARPinComponent.cpp
new file mode 100644
--- /dev/null
+++ b/Engine/Plugins/Lumin/MagicLeap/Source/MagicLeapARPin/Private/MagicLeapARPinComponent.cpp
@@ -0,0 +1,261 @@
+// Copyright 1998-2019 Epic Games, Inc. All Rights Reserved.
+
+#include "MagicLeapARPinComponent.h"
+#include "GameFramework/Actor.h"
+#include "Kismet/KismetMathLibrary.h"
+#include "Kismet/GameplayStatics.h"
+#include "MagicLeapARPinFunctionLibrary.h"
+#include "MagicLeapARPinModule.h"
+#include "IMagicLeapPlugin.h"
+#include "MagicLeapCFUID.h"
+#include "Lumin/CAPIShims/LuminAPI.h"
+
+UMagicLeapARPinComponent::UMagicLeapARPinComponent()
+: UserIndex(0)
+, AutoPinType(EMagicLeapAutoPinType::OnlyOnDataRestoration)
+, bShouldPinActor(false)
+, PinDataClass(UMagicLeapARPinSaveGame::StaticClass())
+, PinnedSceneComponent(nullptr)
+, PinData(nullptr)
+, OldComponentWorldTransform(FTransform::Identity)
+, OldCFUIDTransform(FTransform::Identity)
+, NewComponentWorldTransform(FTransform::Identity)
+, NewCFUIDTransform(FTransform::Identity)
+, bPinned(false)
+, bDataRestored(false)
+, bAttemptedPinningAfterTrackerCreation(false)
+, bPinFoundInEnvironmentPrevFrame(false)
+{
+	// Make sure this component ticks
+	PrimaryComponentTick.bCanEverTick = true;
+	PrimaryComponentTick.bStartWithTickEnabled = true;
+	PrimaryComponentTick.TickGroup = TG_PrePhysics;
+	bAutoActivate = true;
+}
+
+void UMagicLeapARPinComponent::BeginPlay()
+{
+	Super::BeginPlay();
+
+	UMagicLeapARPinFunctionLibrary::CreateTracker();
+
+	if (ObjectUID.Len() == 0)
+	{
+		if (GetOwner() != nullptr)
+		{
+			ObjectUID = GetOwner()->GetName();
+			UE_LOG(LogMagicLeapARPin, Warning, TEXT("ObjectUID is empty. Using Owner actor's name [%s] instead. A non-empty unique ID is required to make the object persistent."), *ObjectUID);
+		}
+	}
+
+	if (ObjectUID.Len() != 0)
+	{
+		PinData = Cast<UMagicLeapARPinSaveGame>(UGameplayStatics::LoadGameFromSlot(ObjectUID, UserIndex));
+		if (PinData)
+		{
+			PinnedCFUID = PinData->PinnedID;
+			OldComponentWorldTransform = PinData->ComponentWorldTransform;
+			OldCFUIDTransform = PinData->PinTransform;
+			bDataRestored = true;
+		}
+	}
+	else
+	{
+		UE_LOG(LogMagicLeapARPin, Warning, TEXT("ObjectUID is empty. A non-empty unique ID is required to make the object persistent."));
+	}
+}
+
+void UMagicLeapARPinComponent::TickComponent(float DeltaTime, enum ELevelTick TickType, FActorComponentTickFunction *ThisTickFunction)
+{
+	Super::TickComponent(DeltaTime, TickType, ThisTickFunction);
+
+	if (!(IMagicLeapPlugin::Get().IsMagicLeapHMDValid() && UMagicLeapARPinFunctionLibrary::IsTrackerValid()))
+	{
+		return;
+	}
+
+	if (!bAttemptedPinningAfterTrackerCreation)
+	{
+		bAttemptedPinningAfterTrackerCreation = true;
+		if ((AutoPinType == EMagicLeapAutoPinType::Always) || (bDataRestored && AutoPinType == EMagicLeapAutoPinType::OnlyOnDataRestoration))
+		{
+			if (bShouldPinActor)
+			{
+				PinActor(GetOwner());
+			}
+			else
+			{
+				PinSceneComponent(this);
+			}
+		}
+	}
+
+	bool bPinFoundInEnvironment = false;
+
+	if (PinnedSceneComponent != nullptr)
+	{
+		if (!bPinned)
+		{
+			if (bDataRestored)
+			{
+				FVector DummyPosition = FVector::ZeroVector;
+				FRotator DummyOrientation = FRotator::ZeroRotator;
+				// use MLSnapshotGetTransform's return value of PoseNotFound to determine of the given PinnedCFUID is available in the current environment.
+				bPinned = UMagicLeapARPinFunctionLibrary::GetARPinPositionAndOrientation(PinnedCFUID, DummyPosition, DummyOrientation, bPinFoundInEnvironment);
+			}
+			else
+			{
+				EMagicLeapPassableWorldError Error = UMagicLeapARPinFunctionLibrary::GetClosestARPin(PinnedSceneComponent->GetComponentLocation(), PinnedCFUID);
+				if (Error == EMagicLeapPassableWorldError::None)
+				{
+					FVector PinWorldPosition = FVector::ZeroVector;
+					FRotator PinWorldOrientation = FRotator::ZeroRotator;
+					bool bCFUIDTransform = UMagicLeapARPinFunctionLibrary::GetARPinPositionAndOrientation(PinnedCFUID, PinWorldPosition, PinWorldOrientation, bPinFoundInEnvironment);
+					if (bCFUIDTransform)
+					{
+						bPinned = true;
+						OldComponentWorldTransform = PinnedSceneComponent->GetComponentToWorld();
+						OldCFUIDTransform = FTransform(PinWorldOrientation, PinWorldPosition);
+
+						if (ObjectUID.Len() != 0)
+						{
+							if (!PinData)
+							{
+								PinData = Cast<UMagicLeapARPinSaveGame>(UGameplayStatics::CreateSaveGameObject(PinDataClass));
+							}
+							check(PinData);
+							PinData->PinnedID = PinnedCFUID;
+							PinData->ComponentWorldTransform = OldComponentWorldTransform;
+							PinData->PinTransform = OldCFUIDTransform;
+							bool bDataSaved = UGameplayStatics::SaveGameToSlot(PinData, ObjectUID, UserIndex);
+							UE_CLOG(!bDataSaved, LogMagicLeapARPin, Error, TEXT("Error saving data persistent data for MagicLeapARPin %s."), *ObjectUID);
+						}
+					}
+				}
+			}
+		}
+
+		if (bPinned)
+		{
+			FVector PinWorldPosition = FVector::ZeroVector;
+			FRotator PinWorldOrientation = FRotator::ZeroRotator;
+			bool bCFUIDTransform = UMagicLeapARPinFunctionLibrary::GetARPinPositionAndOrientation(PinnedCFUID, PinWorldPosition, PinWorldOrientation, bPinFoundInEnvironment);
+			if (bCFUIDTransform)
+			{
+				NewCFUIDTransform = FTransform(PinWorldOrientation.Quaternion(), PinWorldPosition);
+
+				const FMatrix NewComponentWorldTransformMatrix = OldComponentWorldTransform.ToMatrixNoScale() * (OldCFUIDTransform.ToMatrixNoScale().Inverse() * NewCFUIDTransform.ToMatrixNoScale());
+				NewComponentWorldTransform = FTransform(NewComponentWorldTransformMatrix);
+
+				PinnedSceneComponent->SetWorldLocationAndRotation(NewComponentWorldTransform.GetLocation(), NewComponentWorldTransform.Rotator());
+			}
+		}
+	}
+	else
+	{
+		bPinned = false;
+	}
+
+	if (bPinFoundInEnvironmentPrevFrame != bPinFoundInEnvironment)
+	{
+		if (bPinFoundInEnvironment)
+		{
+			OnPersistentEntityPinned.Broadcast(bDataRestored);
+		}
+		else
+		{
+			OnPersistentEntityPinLost.Broadcast();
+		}
+	}
+
+	bPinFoundInEnvironmentPrevFrame = bPinFoundInEnvironment;
+}
+
+void UMagicLeapARPinComponent::EndPlay(const EEndPlayReason::Type EndPlayReason)
+{
+	if (ObjectUID.Len() != 0 && bPinned)
+	{
+		if (!PinData)
+		{
+			PinData = Cast<UMagicLeapARPinSaveGame>(UGameplayStatics::CreateSaveGameObject(PinDataClass));
+		}
+		check(PinData);
+		PinData->PinnedID = PinnedCFUID;
+		PinData->ComponentWorldTransform = NewComponentWorldTransform;
+		PinData->PinTransform = NewCFUIDTransform;
+		bool bDataSaved = UGameplayStatics::SaveGameToSlot(PinData, ObjectUID, UserIndex);
+		UE_CLOG(!bDataSaved, LogMagicLeapARPin, Error, TEXT("Error saving data persistent data for MagicLeapARPin %s."), *ObjectUID);
+	}
+
+	// Reset flags on EndPlay so consequetive VRPreview launches are respected.
+	bAttemptedPinningAfterTrackerCreation = false;
+	bPinFoundInEnvironmentPrevFrame = false;
+
+	Super::EndPlay(EndPlayReason);
+}
+
+bool UMagicLeapARPinComponent::PinSceneComponent(USceneComponent* ComponentToPin)
+{
+	if (ComponentToPin == nullptr)
+	{
+		UE_LOG(LogMagicLeapARPin, Warning, TEXT("nullptr passed to UMagicLeapARPinComponent::PinSceneComponent(). Use UMagicLeapARPinComponent::UnPin() if you no longer wish for this component to be persistent or want to move the component around."));
+		return false;
+	}
+
+	bPinned = false; // Will resolve in Tick
+	PinnedSceneComponent = ComponentToPin;
+	return true;
+}
+
+bool UMagicLeapARPinComponent::PinActor(AActor* ActorToPin)
+{
+	if (ActorToPin != nullptr)
+	{
+		return PinSceneComponent(ActorToPin->GetRootComponent());
+	}
+	return false;
+}
+
+void UMagicLeapARPinComponent::UnPin()
+{
+	PinnedSceneComponent = nullptr;
+	bPinned = false;
+
+	if (ObjectUID.Len() != 0)
+	{
+		UGameplayStatics::DeleteGameInSlot(ObjectUID, UserIndex);
+		PinData = nullptr;
+		bDataRestored = false;
+	}
+}
+
+bool UMagicLeapARPinComponent::IsPinned() const
+{
+	return bPinned;
+}
+
+bool UMagicLeapARPinComponent::PinRestoredOrSynced() const
+{
+	return bDataRestored;
+}
+
+bool UMagicLeapARPinComponent::GetPinnedPinID(FGuid& PinID)
+{
+	if (bPinned)
+	{
+		PinID = PinnedCFUID;
+	}
+	return bPinned;
+}
+
+UMagicLeapARPinSaveGame* UMagicLeapARPinComponent::GetPinData(TSubclassOf<UMagicLeapARPinSaveGame> InPinDataClass)
+{
+	checkf(InPinDataClass == PinDataClass, TEXT("Attempted to get user data object of different type than what was defined in MagicLeapARPinComponent!"));
+
+	if (!PinData)
+	{
+		PinData = Cast<UMagicLeapARPinSaveGame>(UGameplayStatics::CreateSaveGameObject(PinDataClass));
+	}
+	
+	return PinData;
+}
diff --git a/Engine/Plugins/Lumin/MagicLeap/Source/MagicLeapARPin/Private/MagicLeapARPinFunctionLibrary.cpp b/Engine/Plugins/Lumin/MagicLeap/Source/MagicLeapARPin/Private/MagicLeapARPinFunctionLibrary.cpp
new file mode 100644
--- /dev/null
+++ b/Engine/Plugins/Lumin/MagicLeap/Source/MagicLeapARPin/Private/MagicLeapARPinFunctionLibrary.cpp
@@ -0,0 +1,51 @@
+// Copyright 1998-2019 Epic Games, Inc. All Rights Reserved.
+
+#include "MagicLeapARPinFunctionLibrary.h"
+#include "IMagicLeapARPinModule.h"
+
+EMagicLeapPassableWorldError UMagicLeapARPinFunctionLibrary::CreateTracker()
+{
+	return IMagicLeapARPinModule::Get().CreateTracker();
+}
+
+EMagicLeapPassableWorldError UMagicLeapARPinFunctionLibrary::DestroyTracker()
+{
+	return IMagicLeapARPinModule::Get().DestroyTracker();
+}
+
+bool UMagicLeapARPinFunctionLibrary::IsTrackerValid()
+{
+	return IMagicLeapARPinModule::Get().IsTrackerValid();
+}
+
+EMagicLeapPassableWorldError UMagicLeapARPinFunctionLibrary::GetNumAvailableARPins(int32& Count)
+{
+	return IMagicLeapARPinModule::Get().GetNumAvailableARPins(Count);
+}
+
+EMagicLeapPassableWorldError UMagicLeapARPinFunctionLibrary::GetAvailableARPins(int32 NumRequested, TArray<FGuid>& Pins)
+{
+	if (NumRequested <= 0)
+	{
+		GetNumAvailableARPins(NumRequested);
+	}
+
+	if (NumRequested == 0)
+	{
+		// There are no coordinate frames to return, so this call did succeed without any errors, it just returned an array of size 0.
+		Pins.Reset();
+		return EMagicLeapPassableWorldError::None;
+	}
+
+	return IMagicLeapARPinModule::Get().GetAvailableARPins(NumRequested, Pins);
+}
+
+EMagicLeapPassableWorldError UMagicLeapARPinFunctionLibrary::GetClosestARPin(const FVector& SearchPoint, FGuid& PinID)
+{
+	return IMagicLeapARPinModule::Get().GetClosestARPin(SearchPoint, PinID);
+}
+
+bool UMagicLeapARPinFunctionLibrary::GetARPinPositionAndOrientation(const FGuid& PinID, FVector& Position, FRotator& Orientation, bool& PinFoundInEnvironment)
+{
+	return IMagicLeapARPinModule::Get().GetARPinPositionAndOrientation(PinID, Position, Orientation, PinFoundInEnvironment);
+}
diff --git a/Engine/Plugins/Lumin/MagicLeap/Source/MagicLeapARPin/Private/MagicLeapARPinModule.cpp b/Engine/Plugins/Lumin/MagicLeap/Source/MagicLeapARPin/Private/MagicLeapARPinModule.cpp
new file mode 100644
--- /dev/null
+++ b/Engine/Plugins/Lumin/MagicLeap/Source/MagicLeapARPin/Private/MagicLeapARPinModule.cpp
@@ -0,0 +1,284 @@
+// Copyright 1998-2019 Epic Games, Inc. All Rights Reserved.
+
+#include "MagicLeapARPinModule.h"
+#include "MagicLeapMath.h"
+#include "MagicLeapCFUID.h"
+#include "IMagicLeapPlugin.h"
+#include "HeadMountedDisplayFunctionLibrary.h"
+#include "Lumin/CAPIShims/LuminAPIPassableWorld.h"
+#include "Lumin/CAPIShims/LuminAPIPersistentCoordinateFrames.h"
+
+DEFINE_LOG_CATEGORY(LogMagicLeapARPin);
+
+UMagicLeapARPinSaveGame::UMagicLeapARPinSaveGame()
+: PinnedID()
+, ComponentWorldTransform(FTransform::Identity)
+, PinTransform(FTransform::Identity)
+{}
+
+#if WITH_MLSDK
+EMagicLeapPassableWorldError MLToUnrealPassableWorldError(MLResult result)
+{
+	switch (result)
+	{
+		case MLResult_Ok: return EMagicLeapPassableWorldError::None;
+		case MLPassableWorldResult_LowMapQuality: return EMagicLeapPassableWorldError::LowMapQuality;
+		case MLPassableWorldResult_UnableToLocalize: return EMagicLeapPassableWorldError::UnableToLocalize;
+		case MLPassableWorldResult_ServerUnavailable: return EMagicLeapPassableWorldError::Unavailable;
+		case MLResult_PrivilegeDenied: return EMagicLeapPassableWorldError::PrivilegeDenied;
+		case MLResult_InvalidParam: return EMagicLeapPassableWorldError::InvalidParam;
+		case MLResult_UnspecifiedFailure: return EMagicLeapPassableWorldError::UnspecifiedFailure;
+	}
+	return EMagicLeapPassableWorldError::UnspecifiedFailure;
+}
+#endif //WITH_MLSDK
+
+FMagicLeapARPinModule::FMagicLeapARPinModule()
+	: MagicLeap::IAppEventHandler({ EMagicLeapPrivilege::PwFoundObjRead })
+	, bWasTrackerValidOnPause(false)
+	, bCreateTracker(false)
+#if WITH_MLSDK
+	, Tracker(ML_INVALID_HANDLE)
+#endif //WITH_MLSDK
+{}
+
+EMagicLeapPassableWorldError FMagicLeapARPinModule::CreateTracker()
+{
+	EMagicLeapPassableWorldError ErrorCode = EMagicLeapPassableWorldError::Unavailable;
+#if WITH_MLSDK
+	if (!MLHandleIsValid(Tracker))
+	{
+		MagicLeap::EPrivilegeState PrivilegeState = GetPrivilegeStatus(EMagicLeapPrivilege::PwFoundObjRead, false);
+		if (PrivilegeState == MagicLeap::EPrivilegeState::Granted)
+		{
+			// TODO: add retires like in Image tracker if error is LowMapQuality, UnableToLocalize, ServerUnavailable or PrivilegeDenied.
+			// Retrying for PrivilegeDenied would only make sense when MLPrivilege runtime api is functional.
+			MLResult Result = MLPersistentCoordinateFrameTrackerCreate(&Tracker);
+			UE_CLOG(Result != MLResult_Ok, LogMagicLeapARPin, Error, TEXT("Failed to create persistent coordinate frame tracker with error %s."), UTF8_TO_TCHAR(MLPersistentCoordinateFrameGetResultString(Result)));
+			ErrorCode = MLToUnrealPassableWorldError(Result);
+		}
+		else if (PrivilegeState == MagicLeap::EPrivilegeState::Denied)
+		{
+			UE_LOG(LogMagicLeapARPin, Error, TEXT("Failed to initialize persistent coordinate frame tracker due to lack of privilege!"));
+			ErrorCode = EMagicLeapPassableWorldError::PrivilegeDenied;
+		}
+		else if (PrivilegeState == MagicLeap::EPrivilegeState::Pending)
+		{
+			ErrorCode = EMagicLeapPassableWorldError::PrivilegeRequestPending;
+			bCreateTracker = true;
+		}
+	}
+	else
+	{
+		ErrorCode = EMagicLeapPassableWorldError::None;
+	}
+#endif //WITH_MLSDK
+	return ErrorCode;
+}
+
+EMagicLeapPassableWorldError FMagicLeapARPinModule::DestroyTracker()
+{
+	EMagicLeapPassableWorldError ErrorCode = EMagicLeapPassableWorldError::Unavailable;
+#if WITH_MLSDK
+	if (MLHandleIsValid(Tracker))
+	{
+		MLResult Result = MLPersistentCoordinateFrameTrackerDestroy(Tracker);
+		UE_CLOG(Result != MLResult_Ok, LogMagicLeapARPin, Error, TEXT("Failed to destroy persistent coordinate frame tracker with error %s."), UTF8_TO_TCHAR(MLPersistentCoordinateFrameGetResultString(Result)));
+		Tracker = ML_INVALID_HANDLE;
+		ErrorCode = MLToUnrealPassableWorldError(Result);
+		bCreateTracker = false;
+	}
+#endif //WITH_MLSDK
+	return ErrorCode;
+}
+
+bool FMagicLeapARPinModule::IsTrackerValid() const
+{
+#if WITH_MLSDK
+	return MLHandleIsValid(Tracker);
+#else
+	return false;
+#endif //WITH_MLSDK
+}
+
+EMagicLeapPassableWorldError FMagicLeapARPinModule::GetNumAvailableARPins(int32& Count)
+{
+	EMagicLeapPassableWorldError ErrorReturn = CreateTracker();
+	if (ErrorReturn == EMagicLeapPassableWorldError::None)
+	{
+		ErrorReturn = EMagicLeapPassableWorldError::Unavailable;
+		MagicLeap::EPrivilegeState PrivilegeState = GetPrivilegeStatus(EMagicLeapPrivilege::PwFoundObjRead, false);
+		if (PrivilegeState == MagicLeap::EPrivilegeState::Granted)
+		{
+#if WITH_MLSDK
+			uint32 NumPersistentFrames = 0;
+			MLResult Result = MLPersistentCoordinateFrameGetCount(Tracker, &NumPersistentFrames);
+			if (MLResult_Ok == Result)
+			{
+				Count = static_cast<int32>(NumPersistentFrames);
+			}
+			else
+			{
+				UE_LOG(LogMagicLeapARPin, Error, TEXT("MLPersistentCoordinateFrameGetCount failed with error %s"), UTF8_TO_TCHAR(MLPersistentCoordinateFrameGetResultString(Result)));
+				Count = 0;
+			}
+			ErrorReturn = MLToUnrealPassableWorldError(Result);
+#endif // WITH_MLSDK
+		}
+		else if (PrivilegeState == MagicLeap::EPrivilegeState::Pending)
+		{
+			ErrorReturn = EMagicLeapPassableWorldError::PrivilegeRequestPending;
+		}
+		else
+		{
+			ErrorReturn = EMagicLeapPassableWorldError::PrivilegeDenied;
+		}
+	}
+	return ErrorReturn;
+}
+
+EMagicLeapPassableWorldError FMagicLeapARPinModule::GetAvailableARPins(int32 NumRequested, TArray<FGuid>& CoordinateFrames)
+{
+	EMagicLeapPassableWorldError ErrorReturn = CreateTracker();
+	if (ErrorReturn == EMagicLeapPassableWorldError::None)
+	{
+		ErrorReturn = EMagicLeapPassableWorldError::Unavailable;
+		MagicLeap::EPrivilegeState PrivilegeState = GetPrivilegeStatus(EMagicLeapPrivilege::PwFoundObjRead, false);
+		if (PrivilegeState == MagicLeap::EPrivilegeState::Granted)
+		{
+			if (NumRequested > 0)
+			{
+#if WITH_MLSDK
+				CoordinateFrames.Reset();
+				CoordinateFrames.AddZeroed(NumRequested);
+				MLCoordinateFrameUID* ArrayDataPointer = reinterpret_cast<MLCoordinateFrameUID*>(CoordinateFrames.GetData());
+				MLResult QueryResult = MLPersistentCoordinateFrameGetAllEx(Tracker, NumRequested, ArrayDataPointer);
+				if (MLResult_Ok != QueryResult)
+				{
+					UE_LOG(LogMagicLeapARPin, Error, TEXT("MLPersistentCoordinateFrameGetAllEx failed with error %s"), UTF8_TO_TCHAR(MLPersistentCoordinateFrameGetResultString(QueryResult)));
+				}
+				ErrorReturn = MLToUnrealPassableWorldError(QueryResult);
+#endif // WITH_MLSDK
+			}
+		}
+		else if (PrivilegeState == MagicLeap::EPrivilegeState::Pending)
+		{
+			ErrorReturn = EMagicLeapPassableWorldError::PrivilegeRequestPending;
+		}
+		else
+		{
+			ErrorReturn = EMagicLeapPassableWorldError::PrivilegeDenied;
+		}
+	}
+	return ErrorReturn;
+}
+
+EMagicLeapPassableWorldError FMagicLeapARPinModule::GetClosestARPin(const FVector& SearchPoint, FGuid& PinID)
+{
+	EMagicLeapPassableWorldError ErrorReturn = CreateTracker();
+	if (ErrorReturn == EMagicLeapPassableWorldError::None)
+	{
+		ErrorReturn = EMagicLeapPassableWorldError::Unavailable;
+		MagicLeap::EPrivilegeState PrivilegeState = GetPrivilegeStatus(EMagicLeapPrivilege::PwFoundObjRead, false);
+		if (PrivilegeState == MagicLeap::EPrivilegeState::Granted)
+		{
+			const IMagicLeapPlugin& MLPlugin = IMagicLeapPlugin::Get();
+			if (MLPlugin.IsMagicLeapHMDValid())
+			{
+				const float WorldToMetersScale = MLPlugin.GetWorldToMetersScale();
+				const FTransform WorldToTracking = UHeadMountedDisplayFunctionLibrary::GetTrackingToWorldTransform(nullptr).Inverse();
+#if WITH_MLSDK
+				MLVec3f Target = MagicLeap::ToMLVector(WorldToTracking.TransformPosition(SearchPoint), WorldToMetersScale);
+				MLResult Result = MLPersistentCoordinateFrameGetClosest(Tracker, &Target, reinterpret_cast<MLCoordinateFrameUID*>(&PinID));
+				UE_CLOG(MLResult_Ok != Result, LogMagicLeapARPin, Error, TEXT("MLPersistentCoordinateFrameGetClosest failed with error %s"), UTF8_TO_TCHAR(MLPersistentCoordinateFrameGetResultString(Result)));
+				ErrorReturn = MLToUnrealPassableWorldError(Result);
+#endif // WITH_MLSDK
+			}
+		}
+		else if (PrivilegeState == MagicLeap::EPrivilegeState::Pending)
+		{
+			ErrorReturn = EMagicLeapPassableWorldError::PrivilegeRequestPending;
+		}
+		else
+		{
+			ErrorReturn = EMagicLeapPassableWorldError::PrivilegeDenied;
+		}
+	}
+
+	return ErrorReturn;
+}
+
+bool FMagicLeapARPinModule::GetARPinPositionAndOrientation(const FGuid& PinID, FVector& Position, FRotator& Orientation, bool& PinFoundInEnvironment)
+{
+	const IMagicLeapPlugin& MLPlugin = IMagicLeapPlugin::Get();
+
+	if (!MLPlugin.IsMagicLeapHMDValid())
+	{
+		return false;
+	}
+
+	EMagicLeapPassableWorldError ErrorReturn = CreateTracker();
+	if (ErrorReturn == EMagicLeapPassableWorldError::None)
+	{
+		EMagicLeapTransformFailReason FailReason = EMagicLeapTransformFailReason::None;
+		FTransform Pose = FTransform::Identity;
+		if (MLPlugin.GetTransform(PinID, Pose, FailReason))
+		{
+			const FTransform TrackingToWorld = UHeadMountedDisplayFunctionLibrary::GetTrackingToWorldTransform(nullptr);
+			Pose = Pose * TrackingToWorld;
+			Position = Pose.GetLocation();
+			Orientation = Pose.Rotator();
+			PinFoundInEnvironment = true;
+			return true;
+		}
+		PinFoundInEnvironment = (FailReason != EMagicLeapTransformFailReason::PoseNotFound);
+	}
+	return false;
+}
+
+void FMagicLeapARPinModule::OnAppShutDown()
+{
+	DestroyTracker();
+}
+
+void FMagicLeapARPinModule::OnAppTick()
+{
+	if (!IsTrackerValid() && bCreateTracker)
+	{
+		MagicLeap::EPrivilegeState PrivilegeState = GetPrivilegeStatus(EMagicLeapPrivilege::PwFoundObjRead, false);
+		if (PrivilegeState == MagicLeap::EPrivilegeState::Granted)
+		{
+#if WITH_MLSDK
+			// TODO: add retires like in Image tracker if error is LowMapQuality, UnableToLocalize, ServerUnavailable or PrivilegeDenied.
+			// Retrying for PrivilegeDenied would only make sense when MLPrivilege runtime api is functional.
+			MLResult Result = MLPersistentCoordinateFrameTrackerCreate(&Tracker);
+			UE_CLOG(Result != MLResult_Ok, LogMagicLeapARPin, Error, TEXT("Failed to create persistent coordinate frame tracker with error %s."), UTF8_TO_TCHAR(MLPersistentCoordinateFrameGetResultString(Result)));
+#endif // WITH_MLSDK
+			bCreateTracker = false;
+		}
+		else if (PrivilegeState == MagicLeap::EPrivilegeState::Denied)
+		{
+			UE_LOG(LogMagicLeapARPin, Error, TEXT("Failed to initialize persistent coordinate frame tracker due to lack of privilege!"));
+			bCreateTracker = false;
+		}
+	}
+}
+
+void FMagicLeapARPinModule::OnAppPause()
+{
+#if WITH_MLSDK
+	bWasTrackerValidOnPause = MLHandleIsValid(Tracker);
+#endif // WITH_MLSDK
+	// need to destroy the tracker here in case privileges are removed while the app is dormant.
+	DestroyTracker();
+}
+
+void FMagicLeapARPinModule::OnAppResume()
+{
+	if (bWasTrackerValidOnPause)
+	{
+		CreateTracker();
+	}
+}
+
+IMPLEMENT_MODULE(FMagicLeapARPinModule, MagicLeapARPin);
diff --git a/Engine/Plugins/Lumin/MagicLeap/Source/MagicLeapARPin/Private/MagicLeapARPinModule.h b/Engine/Plugins/Lumin/MagicLeap/Source/MagicLeapARPin/Private/MagicLeapARPinModule.h
new file mode 100644
--- /dev/null
+++ b/Engine/Plugins/Lumin/MagicLeap/Source/MagicLeapARPin/Private/MagicLeapARPinModule.h
@@ -0,0 +1,39 @@
+// Copyright 1998-2019 Epic Games, Inc. All Rights Reserved.
+
+#pragma once
+
+#include "CoreMinimal.h"
+#include "IMagicLeapARPinModule.h"
+#include "AppEventHandler.h"
+#include "Lumin/CAPIShims/LuminAPI.h"
+
+DECLARE_LOG_CATEGORY_EXTERN(LogMagicLeapARPin, Verbose, All);
+
+class FMagicLeapARPinModule : public IMagicLeapARPinModule, public MagicLeap::IAppEventHandler
+{
+public:
+	FMagicLeapARPinModule();
+
+	/** IMagicLeapARPinModule interface */
+	virtual EMagicLeapPassableWorldError CreateTracker() override;
+	virtual EMagicLeapPassableWorldError DestroyTracker() override;
+	virtual bool IsTrackerValid() const override;
+	virtual EMagicLeapPassableWorldError GetNumAvailableARPins(int32& Count) override;
+	virtual EMagicLeapPassableWorldError GetAvailableARPins(int32 NumRequested, TArray<FGuid>& Pins) override;
+	virtual EMagicLeapPassableWorldError GetClosestARPin(const FVector& SearchPoint, FGuid& PinID) override;
+	virtual bool GetARPinPositionAndOrientation(const FGuid& PinID, FVector& Position, FRotator& Orientation, bool& PinFoundInEnvironment) override;
+
+	/** IAppEventHandler interface */
+	virtual void OnAppShutDown() override;
+	virtual void OnAppTick() override;
+	virtual void OnAppPause() override;
+	virtual void OnAppResume() override;
+
+private:
+	bool bWasTrackerValidOnPause;
+	bool bCreateTracker;
+
+#if WITH_MLSDK
+	MLHandle Tracker;
+#endif //WITH_MLSDK
+};
diff --git a/Engine/Plugins/Lumin/MagicLeap/Source/MagicLeapARPin/Public/IMagicLeapARPinModule.h b/Engine/Plugins/Lumin/MagicLeap/Source/MagicLeapARPin/Public/IMagicLeapARPinModule.h
new file mode 100644
--- /dev/null
+++ b/Engine/Plugins/Lumin/MagicLeap/Source/MagicLeapARPin/Public/IMagicLeapARPinModule.h
@@ -0,0 +1,84 @@
+// Copyright 1998-2019 Epic Games, Inc. All Rights Reserved.
+#pragma once
+
+#include "Modules/ModuleManager.h"
+#include "MagicLeapARPinTypes.h"
+
+	/**
+	 * The public interface to this module.  In most cases, this interface is only public to sibling modules
+	 * within this plugin.
+	 */
+class MAGICLEAPARPIN_API IMagicLeapARPinModule : public IModuleInterface
+{
+public:
+
+	/**
+	 * Singleton-like access to this module's interface.  This is just for convenience!
+	 * Beware of calling this during the shutdown phase, though.  Your module might have been unloaded already.
+	 *
+	 * @return Returns singleton instance, loading the module on demand if needed
+	 */
+	static inline IMagicLeapARPinModule& Get()
+	{
+		return FModuleManager::LoadModuleChecked<IMagicLeapARPinModule>("MagicLeapARPin");
+	}
+
+	/**
+	 * Checks to see if this module is loaded and ready.  It is only valid to call Get() if IsAvailable() returns true.
+	 *
+	 * @return True if the module is loaded and ready to use
+	 */
+	static inline bool IsAvailable()
+	{
+		return FModuleManager::Get().IsModuleLoaded("MagicLeapARPin");
+	}
+
+	/**
+	* Create an ARPin tracker.
+	* @return Error code representing specific success or failure cases. If code is EMagicLeapPassableWorldError::PrivilegeRequestPending,
+	* poll for IsTrackerValid() to check when the privilege is granted and tracker successfully created.
+	*/
+	virtual EMagicLeapPassableWorldError CreateTracker() = 0;
+
+	/**
+	* Destroy an ARPin tracker.
+	* @return Error code representing specific success or failure cases.,
+	*/
+	virtual EMagicLeapPassableWorldError DestroyTracker() = 0;
+
+	/** Is an ARPin tracker already created. */
+	virtual bool IsTrackerValid() const = 0;
+
+	/**
+	* Returns the count of currently available AR Pins.
+	* @param Count Output param for number of currently available AR Pins. Valid only if return value is EMagicLeapPassableWorldError::None.
+	* @return Error code representing specific success or failure cases.
+	*/
+	virtual EMagicLeapPassableWorldError GetNumAvailableARPins(int32& Count) = 0;
+
+	/**
+	* Returns all the AR Pins currently available.
+	* @param NumRequested Max number of AR Pins to query. Pass in a negative integer to get all available Pins.
+	* @param Pins Output array containing IDs of the found Pins. Valid only if return value is EMagicLeapPassableWorldError::None.
+	* @return Error code representing specific success or failure cases.
+	*/
+	virtual EMagicLeapPassableWorldError GetAvailableARPins(int32 NumRequested, TArray<FGuid>& Pins) = 0;
+
+	/**
+	* Returns the Pin closest to the target point passed in.
+	* @param SearchPoint Position, in world space, to search the closest Pin to.
+	* @param PinID Output param for the ID of the closest Pin. Valid only if return value is EMagicLeapPassableWorldError::None.
+	* @return Error code representing specific success or failure cases.
+	*/
+	virtual EMagicLeapPassableWorldError GetClosestARPin(const FVector& SearchPoint, FGuid& PinID) = 0;
+
+	/**
+	* Returns the world position & orientation of the requested Pin.
+	* @param PinID ID of the Pin to get the position and orientation for.
+	* @param Position Output param for the world position of the Pin. Valid only if return value is true.
+	* @param Orientation Output param for the world orientation of the Pin. Valid only if return value is true.
+	* @param PinFoundInEnvironment Output param for indicating if the requested Pin was found user's current environment or not.
+	* @return true if the PinID was valid and the position & orientation were successfully retrieved.
+	*/
+	virtual bool GetARPinPositionAndOrientation(const FGuid& PinID, FVector& Position, FRotator& Orientation, bool& PinFoundInEnvironment) = 0;
+};
diff --git a/Engine/Plugins/Lumin/MagicLeap/Source/MagicLeapARPin/Public/MagicLeapARPinComponent.h b/Engine/Plugins/Lumin/MagicLeap/Source/MagicLeapARPin/Public/MagicLeapARPinComponent.h
new file mode 100644
--- /dev/null
+++ b/Engine/Plugins/Lumin/MagicLeap/Source/MagicLeapARPin/Public/MagicLeapARPinComponent.h
@@ -0,0 +1,150 @@
+// Copyright 1998-2019 Epic Games, Inc. All Rights Reserved.
+#pragma once
+
+#include "Components/SceneComponent.h"
+#include "MagicLeapARPinTypes.h"
+#include "MagicLeapARPinComponent.generated.h"
+
+class AActor;
+
+/** Component to make content persist at locations in the real world. */
+UCLASS(ClassGroup = MagicLeap, BlueprintType, Blueprintable, EditInlineNew, meta = (BlueprintSpawnableComponent))
+class MAGICLEAPARPIN_API UMagicLeapARPinComponent : public USceneComponent
+{
+	GENERATED_BODY()
+
+public:
+	UMagicLeapARPinComponent();
+
+	virtual void BeginPlay() override;
+	virtual void TickComponent(float DeltaTime, enum ELevelTick TickType, FActorComponentTickFunction *ThisTickFunction) override;
+	virtual void EndPlay(const EEndPlayReason::Type EndPlayReason) override;
+
+	/**
+	 * Pin given SceneComponent to the closest AR Pin in real-world.
+	 * OnPersistentEntityPinned event will be fired when a suitable AR Pin is found for this component.
+	 * The component's transform will then be locked. App needs to call UnPin() if it wants to move the component again.
+	 * @param ComponentToPin SceneComponent to pin to the world. Pass in 'this' component if app is using 'OnlyOnDataRestoration' or 'Always' AutoPinType.
+	 * @return true if the component was accepted to be pinned, false otherwise.
+	 */
+	UFUNCTION(BlueprintCallable, Category = "ContentPersistence|MagicLeap")
+	bool PinSceneComponent(USceneComponent* ComponentToPin);
+
+	/**
+	 * Pin given Actor to the closest AR Pin in real-world.
+	 * OnPersistentEntityPinned event will be fired when a suitable AR Pin is found for this Actor.
+	 * The Actor's transform will then be locked. App needs to call UnPin() if it wants to move the Actor again.
+	 * @param ActorToPin Actor to pin to the world. Pass in this component's owner if app is using 'OnlyOnDataRestoration' or 'Always' AutoPinType.
+	 * @return true if the Actor was accepted to be pinned, false otherwise.
+	 */
+	UFUNCTION(BlueprintCallable, Category = "ContentPersistence|MagicLeap")
+	bool PinActor(AActor* ActorToPin);
+
+	/**
+	 * Detach or un-pin the currently pinned entity (component or actor) from the real-world.
+	 * Call this if you want to change the transform of a pinned entity.
+	 * Note that if you still want your content to persist, you will have to call PinSceneComponent() or PinActor() before EndPlay().
+	 */
+	UFUNCTION(BlueprintCallable, Category = "ContentPersistence|MagicLeap")
+	void UnPin();
+
+	/**
+	 * True if an entity (component or actor) is currently pinned by this component.
+	 * If true, the entity's transform will be locked. App needs to call UnPin() if it wants to move it again.
+	 * If false, and you still want your content to persist, you will have to call PinSceneComponent() or PinActor() before EndPlay().
+	 * @return True if an entity (component or actor) is currently pinned by this component.
+	 */
+	UFUNCTION(BlueprintCallable, Category = "ContentPersistence|MagicLeap")
+	bool IsPinned() const;
+
+	/**
+	 * True if the AR Pin for the unique ID ObjectUID was restored from the app's local storage or was repliated over network.
+	 * Implies if content was already pinned earlier. Does not imply if that restored Pin is available in the current environment.
+	 * @return True if the Pin data was restored from local storage or network.
+	 */
+	UFUNCTION(BlueprintCallable, Category = "ContentPersistence|MagicLeap")
+	bool PinRestoredOrSynced() const;
+
+	/**
+	 * Get the ID of the Pin the entity (component or actor) is currently pinned to.
+	 * @param PinID Output param for the ID of the Pin.
+	 * @return True if an entity is currently pinned by this component and the output param is successfully populated.
+	 */
+	UFUNCTION(BlueprintCallable, Category = "ContentPersistence|MagicLeap")
+	bool GetPinnedPinID(FGuid& PinID);
+
+	/**
+	 * Retrieves the data associated with this pin.
+	 * @param PinDataClass The user defined save game class used by this pin.  Note that this must match the PinDataClass property.
+	 * @return The save game instance associated with this pin instance.
+	 */
+	UFUNCTION(BlueprintCallable, Category = "ContentPersistence|MagicLeap", meta = (DeterminesOutputType = "PinDataClass"))
+	UMagicLeapARPinSaveGame* GetPinData(TSubclassOf<UMagicLeapARPinSaveGame> PinDataClass);
+
+public:
+	/**
+	 * Unique ID for this component to save the meta data for the Pin and make content persistent.
+	 * This name has to be unique across all instances of the MagicLeapARPinComponent class.
+	 * If empty, the name of the owner actor will be used.
+	 */
+	UPROPERTY(BlueprintReadWrite, EditAnywhere, Category = "ContentPersistence|MagicLeap")
+	FString ObjectUID;
+
+	/** Index to get the save game data for the pin */
+	UPROPERTY(BlueprintReadWrite, EditAnywhere, Category = "ContentPersistence|MagicLeap")
+	int32 UserIndex;
+
+	/** Mode for automatically pinning this component or it's owner actor to real-world. */
+	UPROPERTY(BlueprintReadWrite, EditAnywhere, Category = "ContentPersistence|MagicLeap")
+	EMagicLeapAutoPinType AutoPinType;
+
+	/** Pin this component's owner actor instead of just the component itself. Relevant only when using 'OnlyOnDataRestoration' or 'Always' as AutoPinType. */
+	UPROPERTY(BlueprintReadWrite, EditAnywhere, Category = "ContentPersistence|MagicLeap")
+	bool bShouldPinActor;
+
+	/** The user defined save game class associated with this pin.  Note that this MUST match the type passed into GetPinData().*/
+	UPROPERTY(BlueprintReadWrite, EditAnywhere, Category = "ContentPersistence|MagicLeap")
+	TSubclassOf<UMagicLeapARPinSaveGame> PinDataClass;
+
+	/**
+	 * Delegate used to notify the instigating blueprint that an entity (component or actor) has been successfully pinned to the real-world.
+	 * Indicates that the transform of the pinned entity is now locked. App needs to call UnPin() if it wants to move the entity again.
+	 * @param bRestoredOrSynced True if the entity was pinned as a result of Pin data being restored from local storage or replicatred over network, false if pinned by an explicit PinSceneComponent() or PinActor() call from the app.
+	 */
+	DECLARE_DYNAMIC_MULTICAST_DELEGATE_OneParam(FPersistentEntityPinned, bool, bRestoredOrSynced);
+
+	/** Fired when an entity is successfully pinned by this component. */
+	UPROPERTY(BlueprintAssignable)
+	FPersistentEntityPinned OnPersistentEntityPinned;
+
+	/**
+	* Delegate used to notify the instigating blueprint that an entity (component or actor) has lost a previously obtained pin.
+	*/
+	DECLARE_DYNAMIC_MULTICAST_DELEGATE(FPersistentEntityPinLost);
+
+	/** Fired when an entity loses its pin. */
+	UPROPERTY(BlueprintAssignable)
+	FPersistentEntityPinLost OnPersistentEntityPinLost;
+
+private:
+	UPROPERTY()
+	FGuid PinnedCFUID;
+
+	UPROPERTY()
+	USceneComponent* PinnedSceneComponent;
+
+	UPROPERTY()
+	UMagicLeapARPinSaveGame* PinData;
+
+	FTransform OldComponentWorldTransform;
+	FTransform OldCFUIDTransform;
+	FTransform NewComponentWorldTransform;
+	FTransform NewCFUIDTransform;
+
+	bool bPinned;
+	bool bDataRestored;
+	bool bAttemptedPinningAfterTrackerCreation;
+	bool bPinFoundInEnvironmentPrevFrame;
+
+	class FMagicLeapARPinTrackerImpl *Impl;
+};
diff --git a/Engine/Plugins/Lumin/MagicLeap/Source/MagicLeapARPin/Public/MagicLeapARPinFunctionLibrary.h b/Engine/Plugins/Lumin/MagicLeap/Source/MagicLeapARPin/Public/MagicLeapARPinFunctionLibrary.h
new file mode 100644
--- /dev/null
+++ b/Engine/Plugins/Lumin/MagicLeap/Source/MagicLeapARPin/Public/MagicLeapARPinFunctionLibrary.h
@@ -0,0 +1,70 @@
+// Copyright 1998-2019 Epic Games, Inc. All Rights Reserved.
+#pragma once
+
+#include "Kismet/BlueprintFunctionLibrary.h"
+#include "MagicLeapARPinTypes.h"
+#include "MagicLeapARPinFunctionLibrary.generated.h"
+
+/** Direct API interface for the Magic Leap Persistent AR Pin tracker system. */
+UCLASS(ClassGroup = MagicLeap)
+class MAGICLEAPARPIN_API UMagicLeapARPinFunctionLibrary : public UBlueprintFunctionLibrary
+{
+	GENERATED_BODY()
+
+public:
+	/**
+	* Create an ARPin tracker.
+	* @return Error code representing specific success or failure cases. If code is EMagicLeapPassableWorldError::PrivilegeRequestPending,
+	* poll for IsTrackerValid() to check when the privilege is granted and tracker successfully created.
+	*/
+	UFUNCTION(BlueprintCallable, Category = "ContentPersistence|MagicLeap")
+	static EMagicLeapPassableWorldError CreateTracker();
+
+	/**
+	* Destroy an ARPin tracker.
+	* @return Error code representing specific success or failure cases.,
+	*/
+	UFUNCTION(BlueprintCallable, Category = "ContentPersistence|MagicLeap")
+	static EMagicLeapPassableWorldError DestroyTracker();
+
+	/** Is an ARPin tracker already created. */
+	UFUNCTION(BlueprintCallable, BlueprintPure, Category = "ContentPersistence|MagicLeap")
+	static bool IsTrackerValid();
+
+	/**
+	* Returns the count of currently available AR Pins.
+	* @param Count Output param for number of currently available AR Pins. Valid only if return value is EMagicLeapPassableWorldError::None.
+	* @return Error code representing specific success or failure cases.
+	*/
+	UFUNCTION(BlueprintCallable, Category = "ContentPersistence|MagicLeap")
+	static EMagicLeapPassableWorldError GetNumAvailableARPins(int32& Count);
+
+	/**
+	* Returns all the AR Pins currently available.
+	* @param NumRequested Max number of AR Pins to query. Pass in a negative integer to get all available Pins.
+	* @param Pins Output array containing IDs of the found Pins. Valid only if return value is EMagicLeapPassableWorldError::None.
+	* @return Error code representing specific success or failure cases.
+	*/
+	UFUNCTION(BlueprintCallable, Category = "ContentPersistence|MagicLeap")
+	static EMagicLeapPassableWorldError GetAvailableARPins(int32 NumRequested, TArray<FGuid>& Pins);
+
+	/**
+	* Returns the Pin closest to the target point passed in.
+	* @param SearchPoint Position, in world space, to search the closest Pin to.
+	* @param PinID Output param for the ID of the closest Pin. Valid only if return value is EMagicLeapPassableWorldError::None.
+	* @return Error code representing specific success or failure cases.
+	*/
+	UFUNCTION(BlueprintCallable, Category = "ContentPersistence|MagicLeap")
+	static EMagicLeapPassableWorldError GetClosestARPin(const FVector& SearchPoint, FGuid& PinID);
+
+	/**
+	* Returns the world position & orientation of the requested Pin.
+	* @param PinID ID of the Pin to get the position and orientation for.
+	* @param Position Output param for the world position of the Pin. Valid only if return value is true.
+	* @param Orientation Output param for the world orientation of the Pin. Valid only if return value is true.
+	* @param PinFoundInEnvironment Output param for indicating if the requested Pin was found user's current environment or not.
+	* @return true if the PinID was valid and the position & orientation were successfully retrieved.
+	*/
+	UFUNCTION(BlueprintCallable, Category = "ContentPersistence|MagicLeap")
+	static bool GetARPinPositionAndOrientation(const FGuid& PinID, FVector& Position, FRotator& Orientation, bool& PinFoundInEnvironment);
+};
diff --git a/Engine/Plugins/Lumin/MagicLeap/Source/MagicLeapARPin/Public/MagicLeapARPinTypes.h b/Engine/Plugins/Lumin/MagicLeap/Source/MagicLeapARPin/Public/MagicLeapARPinTypes.h
new file mode 100644
--- /dev/null
+++ b/Engine/Plugins/Lumin/MagicLeap/Source/MagicLeapARPin/Public/MagicLeapARPinTypes.h
@@ -0,0 +1,63 @@
+// Copyright 1998-2019 Epic Games, Inc. All Rights Reserved.
+#pragma once
+
+#include "Engine/Engine.h"
+#include "UObject/NoExportTypes.h"
+#include "GameFramework/SaveGame.h"
+#include "MagicLeapARPinTypes.generated.h"
+
+/** List of possible error values for MagicLeapARPin functions. */
+UENUM(BlueprintType)
+enum class EMagicLeapPassableWorldError : uint8
+{
+	/** No error. */
+	None,
+	/** Map quality too low for content persistence. Continue building the map. */
+	LowMapQuality,
+	/** Currently unable to localize into any map. Continue building the map. */
+	UnableToLocalize,
+	/** AR Pin is not available at this time. */
+	Unavailable,
+	/** Privileges not met. Add 'PwFoundObjRead' privilege to app manifest and request it at runtime. */
+	PrivilegeDenied,
+	/** Invalid function parameter. */
+	InvalidParam,
+	/** Unspecified error. */
+	UnspecifiedFailure,
+	/** Privilege has been requested but not yet granted by the user. */
+	PrivilegeRequestPending
+};
+
+/** Modes for automatically pinning content to real-world. */
+UENUM(BlueprintType)
+enum class EMagicLeapAutoPinType : uint8
+{
+	/** 
+	 * Pin this component / owner actor automatically only if it was pinned in a previous run of the app or replicated over network.
+	 * App needs to call PinSceneComponent() or PinActor() to pin for the very first time.
+	 */
+	OnlyOnDataRestoration,
+	/** Always pin this component / owner actor automatically, without having to call PinSceneComponent() or PinActor() explicitely. */
+	Always,
+	/** Never pin this component / owner actor automatically. App will control pinning and unpinning itself. */
+	Never
+};
+
+/** Base class to save game data associated with a given pin. Inherit from this class and set it to be the PinDataClass in the MagicLeapARPinComponent to save and restore related data. */
+UCLASS(ClassGroup = MagicLeap, BlueprintType, Blueprintable)
+class MAGICLEAPARPIN_API UMagicLeapARPinSaveGame : public USaveGame
+{
+	GENERATED_BODY()
+
+public:
+	UMagicLeapARPinSaveGame();
+
+	UPROPERTY(VisibleAnywhere, Category = "ContentPersistence|MagicLeap")
+	FGuid PinnedID;
+
+	UPROPERTY(VisibleAnywhere, Category = "ContentPersistence|MagicLeap")
+	FTransform ComponentWorldTransform;
+
+	UPROPERTY(VisibleAnywhere, Category = "ContentPersistence|MagicLeap")
+	FTransform PinTransform;
+};
diff --git a/Engine/Plugins/Lumin/MagicLeap/Source/MagicLeapAudio/MagicLeapAudio.Build.cs b/Engine/Plugins/Lumin/MagicLeap/Source/MagicLeapAudio/MagicLeapAudio.Build.cs
--- a/Engine/Plugins/Lumin/MagicLeap/Source/MagicLeapAudio/MagicLeapAudio.Build.cs
+++ b/Engine/Plugins/Lumin/MagicLeap/Source/MagicLeapAudio/MagicLeapAudio.Build.cs
@@ -13,8 +13,6 @@ public class MagicLeapAudio : ModuleRules
 		PublicIncludePaths.Add(Path.Combine(EngineDir, "Source/Runtime/AudioMixer/Public"));
 		PrivateIncludePaths.Add(Path.Combine(EngineDir, "Source/Runtime/AudioMixer/Private"));
 
-		PrivateIncludePaths.Add("MagicLeap/Private");
-
 		PrivateDependencyModuleNames.AddRange(
 			new string[] {
 				"Core",
@@ -22,7 +20,9 @@ public class MagicLeapAudio : ModuleRules
 				"Engine",
 				"MLSDK",
 				"AudioMixer",
-				"AudioMixerCore"
+				"AudioMixerCore",
+				"HeadMountedDisplay",
+				"MagicLeap"
 			}
 		);
 
