Copying //UE4/Partner-MagicLeap-Staging to Dev-VR (//UE4/Dev-VR)

diff --git a/Engine/Plugins/Lumin/MagicLeap/Source/MagicLeapAudio/Private/AudioMixerPlatformMagicLeap.cpp b/Engine/Plugins/Lumin/MagicLeap/Source/MagicLeapAudio/Private/AudioMixerPlatformMagicLeap.cpp
--- a/Engine/Plugins/Lumin/MagicLeap/Source/MagicLeapAudio/Private/AudioMixerPlatformMagicLeap.cpp
+++ b/Engine/Plugins/Lumin/MagicLeap/Source/MagicLeapAudio/Private/AudioMixerPlatformMagicLeap.cpp
@@ -9,7 +9,11 @@
 #include "Misc/CoreDelegates.h"
 #include "VorbisAudioInfo.h"
 #include "ADPCMAudioInfo.h"
-
+#include "MagicLeapAudioModule.h"
+#include "Lumin/CAPIShims/LuminAPILogging.h"
+#if PLATFORM_LUMIN
+#include "Lumin/LuminPlatformDelegates.h"
+#endif // PLATFORM_LUMIN
 
 DECLARE_LOG_CATEGORY_EXTERN(LogAudioMixerMagicLeap, Log, All);
 DEFINE_LOG_CATEGORY(LogAudioMixerMagicLeap);
@@ -34,7 +38,7 @@ DEFINE_LOG_CATEGORY(LogAudioMixerMagicLeap);
 #define MLAUDIO_LOG_FAILURE(Result)							\
 	{														\
 		const TCHAR* ErrorString = FMixerPlatformMagicLeap::GetErrorString(Result);	\
-		UE_LOG(LogAudioMixerMagicLeap, Error, TEXT("Error in %s, line %s: %s"), __FILE__, __LINE__,ErrorString); \
+		UE_LOG(LogAudioMixerMagicLeap, Error, TEXT("Error in %s, line %d: %s"), __FILE__, __LINE__, ErrorString); \
 	}
 
 namespace Audio
@@ -45,11 +49,11 @@ namespace Audio
 	const uint32 DefaultSamplesPerSecond = 48000;  // presumed 48KHz and 16 bits for the sample
 	// TODO: @Epic check the value to be used. Setting default for now.
 	const float DefaultMaxPitch = 1.0f;
+	static bool bRetrievedDeviceDefaults = false;
 
 #if WITH_MLSDK
 	static void GetMLDeviceDefaults(MLAudioBufferFormat& OutFormat, uint32& OutSize, uint32& OutMinSize)
 	{
-		static bool bRetrievedDeviceDefaults = false;
 		static uint32 CachedSize = 0;
 		static uint32 CachedMinSize = 0;
 		static MLAudioBufferFormat CachedBufferFormat;
@@ -64,7 +68,7 @@ namespace Audio
 			else
 			{
 				bRetrievedDeviceDefaults = false;
-				UE_LOG(LogAudioMixerMagicLeap, Error, TEXT("MLAudioGetOutputStreamDefaults failed with error %d."), Result);
+				UE_LOG(LogAudioMixerMagicLeap, Error, TEXT("MLAudioGetOutputStreamDefaults failed with error %s."), UTF8_TO_TCHAR(MLAudioGetResultString(Result)));
 			}
 		}
 
@@ -79,10 +83,15 @@ namespace Audio
 		, bSuspended(false)
 		, bInitialized(false)
 		, bInCallback(false)
+		, FakeCallback(this)
 #if WITH_MLSDK
 		, StreamHandle(ML_INVALID_HANDLE)
 #endif //WITH_MLSDK
 	{
+#if WITH_MLSDK && PLATFORM_MAC
+		// Force load ml_ext_logging before any ml_audio function is called because of REM-3398
+		MLLoggingLogLevelIsEnabled(MLLogLevel_Error);
+#endif
 	}
 
 	FMixerPlatformMagicLeap::~FMixerPlatformMagicLeap()
@@ -128,6 +137,12 @@ namespace Audio
 		FCoreDelegates::ApplicationWillEnterBackgroundDelegate.AddRaw(this, &FMixerPlatformMagicLeap::SuspendContext);
 		FCoreDelegates::ApplicationHasEnteredForegroundDelegate.AddRaw(this, &FMixerPlatformMagicLeap::ResumeContext);
 
+#if PLATFORM_LUMIN		
+		FLuminDelegates::DeviceWillEnterRealityModeDelegate.AddRaw(this, &FMixerPlatformMagicLeap::DeviceStandby);
+		FLuminDelegates::DeviceWillGoInStandbyDelegate.AddRaw(this, &FMixerPlatformMagicLeap::DeviceStandby);
+		FLuminDelegates::DeviceHasReactivatedDelegate.AddRaw(this, &FMixerPlatformMagicLeap::DeviceActive);
+#endif // PLATFORM_LUMIN
+
 		bInitialized = true;
 
 		return true;
@@ -144,6 +159,12 @@ namespace Audio
 		FCoreDelegates::ApplicationWillEnterBackgroundDelegate.RemoveAll(this);
 		FCoreDelegates::ApplicationHasEnteredForegroundDelegate.RemoveAll(this);
 
+#if PLATFORM_LUMIN
+		FLuminDelegates::DeviceWillEnterRealityModeDelegate.RemoveAll(this);
+		FLuminDelegates::DeviceWillGoInStandbyDelegate.RemoveAll(this);
+		FLuminDelegates::DeviceHasReactivatedDelegate.RemoveAll(this);
+#endif // PLATFORM_LUMIN
+
 		bInitialized = false;
 		return true;
 	}
@@ -169,6 +190,7 @@ namespace Audio
 
 		GetMLDeviceDefaults(DesiredBufferFormat, OutSize, OutMinimalSize);
 
+		check(bRetrievedDeviceDefaults == true);
 		OutInfo.Name = TEXT("Magic Leap Audio Device");
 		OutInfo.DeviceId = 0;
 		OutInfo.bIsSystemDefault = true;
@@ -222,6 +244,7 @@ namespace Audio
 		uint32 OutMinimalSize;
 
 		GetMLDeviceDefaults(DesiredBufferFormat, OutSize, OutMinimalSize);
+		check(bRetrievedDeviceDefaults == true);
 
 		OpenStreamParams = Params;
 
@@ -256,23 +279,26 @@ namespace Audio
 	bool FMixerPlatformMagicLeap::CloseAudioStream()
 	{
 #if WITH_MLSDK
-		check(MLHandleIsValid(StreamHandle));
-
-		if (!bInitialized || (AudioStreamInfo.StreamState != EAudioOutputStreamState::Open && AudioStreamInfo.StreamState != EAudioOutputStreamState::Stopped))
 		{
-			return false;
-		}
+			FScopeLock Lock(&StreamHandleCriticalSection);
+			check(MLHandleIsValid(StreamHandle));
 
-		MLResult Result = MLAudioDestroySound(StreamHandle);
+			if (!bInitialized || (AudioStreamInfo.StreamState != EAudioOutputStreamState::Open && AudioStreamInfo.StreamState != EAudioOutputStreamState::Stopped))
+			{
+				return false;
+			}
 
-		if (Result != MLResult_Ok)
-		{
-			MLAUDIO_LOG_FAILURE(Result);
-			return false;
-		}
+			MLResult Result = MLAudioDestroySound(StreamHandle);
 
-		AudioStreamInfo.StreamState = EAudioOutputStreamState::Closed;
-		StreamHandle = ML_INVALID_HANDLE;
+			if (Result != MLResult_Ok)
+			{
+				MLAUDIO_LOG_FAILURE(Result);
+				return false;
+			}
+
+			AudioStreamInfo.StreamState = EAudioOutputStreamState::Closed;
+			StreamHandle = ML_INVALID_HANDLE;
+		}
 #endif //WITH_MLSDK
 
 		return true;
@@ -360,6 +386,7 @@ namespace Audio
 
 		GetMLDeviceDefaults(DesiredBufferFormat, OutSize, OutMinimalSize);
 
+		check(bRetrievedDeviceDefaults == true);
 		FAudioPlatformSettings PlatformSettings;
 		PlatformSettings.CallbackBufferFrameSize = OutSize / (DefaultNumChannels * (DesiredBufferFormat.bits_per_sample / 8));
 		PlatformSettings.MaxChannels = 0;
@@ -440,6 +467,11 @@ namespace Audio
 #endif //WITH_MLSDK
 	}
 
+	void FMixerPlatformMagicLeap::OnHardwareUpdate()
+	{
+		GetMagicLeapAudioModule().CheckOutputDevice();
+	}
+
 	int32 FMixerPlatformMagicLeap::GetNumFrames(const int32 InNumReqestedFrames)
 	{
 #if WITH_MLSDK
@@ -448,6 +480,7 @@ namespace Audio
 		uint32 OutMinimalSize;
 
 		GetMLDeviceDefaults(DesiredBufferFormat, OutSize, OutMinimalSize);
+		check(bRetrievedDeviceDefaults == true);
 		return OutSize / (DefaultNumChannels * (DesiredBufferFormat.bits_per_sample / 8));
 #else
 		return 0;
@@ -474,48 +507,75 @@ namespace Audio
 #endif //WITH_MLSDK
 	}
 
+	void FMixerPlatformMagicLeap::DeviceStandby()
+	{
+		// When the device goes in standby mode, the MLAudioCallback function is not called.
+		// As a result the audio buffers from the engine queue up and they play back once the device becomes active again.
+		// This desyncs gameplay with audio and breaks the UX spec for standby mode. When in standby, the app should function normally
+		// from a gameplay standpoint. It should not "pause". This spec is to ensure that when the device is active again, the "resume" time is very minimal.
+		// Using this fake callback, we keep emptying the audio buffer from the engine and act as if everything is running like it normally should.
+		// Another aspect this fixes is launching an app when device is already in standby mode. If during engine initialization, the audio mixer
+		// does not call IAudioMixerPlatformInterface::ReadNextBuffer(), the audio engine blocks and eventually crashes.
+		// This fake callback ensures a smooth initialization as well. 
+		FakeCallback.SetShouldUseCallback(true);
+	}
+
+	void FMixerPlatformMagicLeap::DeviceActive()
+	{
+		FakeCallback.SetShouldUseCallback(false);
+	}
+
 #if WITH_MLSDK
 	void FMixerPlatformMagicLeap::MLAudioCallback(MLHandle Handle, void* CallbackContext)
 	{
-		FMixerPlatformMagicLeap* InPlatform = (FMixerPlatformMagicLeap*)CallbackContext;
-
-		MLAudioBuffer CallbackBuffer;
+		FMixerPlatformMagicLeap* InPlatform = reinterpret_cast<FMixerPlatformMagicLeap*>(CallbackContext);
 
-		check(MLHandleIsValid(InPlatform->StreamHandle));
-		// TODO: should this be an assert?
-		if (Handle != InPlatform->StreamHandle)
+		if (InPlatform == nullptr)
 		{
+			UE_LOG(LogAudioMixerMagicLeap, Error, TEXT("The callback context to MLAudioBufferCallback for MLAudioCreateSoundWithOutputStream was null!"));
 			return;
 		}
-		// Get the callback buffer from MLAudio
-		MLResult Result = MLAudioGetOutputStreamBuffer(InPlatform->StreamHandle, &CallbackBuffer);
-		if (Result != MLResult_Ok)
-		{
-			MLAUDIO_LOG_FAILURE(Result);
-			return;
-		}
+		check(MLHandleIsValid(Handle));
 
-		if (InPlatform->CachedBufferHandle == nullptr)
+		// This handle may be used even after being invalidated in CloseAudioStream's call to MLAudioDestroySound
 		{
-			InPlatform->ReadNextBuffer();
-		}
+			FScopeLock Lock(&InPlatform->StreamHandleCriticalSection);
+			if (!MLHandleIsValid(InPlatform->StreamHandle))
+			{
+				return;
+			}
 
-		// It is possible that ReadNextBuffer() doesn't call SubmitBuffer()
-		// in which case CachedBufferHandle will still be null and memcpy will segfault.
-		if (InPlatform->CachedBufferHandle != nullptr)
-		{
-			// Fill the callback buffer:
-			FMemory::Memcpy(CallbackBuffer.ptr, InPlatform->CachedBufferHandle, CallbackBuffer.size);
-		}
+			MLAudioBuffer CallbackBuffer;
+			// Get the callback buffer from MLAudio
+			MLResult Result = MLAudioGetOutputStreamBuffer(InPlatform->StreamHandle, &CallbackBuffer);
+			if (Result != MLResult_Ok)
+			{
+				MLAUDIO_LOG_FAILURE(Result);
+				return;
+			}
 
-		Result = MLAudioReleaseOutputStreamBuffer(InPlatform->StreamHandle);
-		if (Result != MLResult_Ok)
-		{
-			MLAUDIO_LOG_FAILURE(Result);
-			return;
-		}
+			if (InPlatform->CachedBufferHandle == nullptr)
+			{
+				InPlatform->ReadNextBuffer();
+			}
 
-		InPlatform->CachedBufferHandle = nullptr;
+			// It is possible that ReadNextBuffer() doesn't call SubmitBuffer()
+			// in which case CachedBufferHandle will still be null and memcpy will segfault.
+			if (InPlatform->CachedBufferHandle != nullptr)
+			{
+				// Fill the callback buffer:
+				FMemory::Memcpy(CallbackBuffer.ptr, InPlatform->CachedBufferHandle, CallbackBuffer.size);
+			}
+
+			Result = MLAudioReleaseOutputStreamBuffer(InPlatform->StreamHandle);
+			if (Result != MLResult_Ok)
+			{
+				MLAUDIO_LOG_FAILURE(Result);
+				return;
+			}
+
+			InPlatform->CachedBufferHandle = nullptr;
+		}
 	}
 #endif //WITH_MLSDK
 }
diff --git a/Engine/Plugins/Lumin/MagicLeap/Source/MagicLeapAudio/Private/AudioMixerPlatformMagicLeap.h b/Engine/Plugins/Lumin/MagicLeap/Source/MagicLeapAudio/Private/AudioMixerPlatformMagicLeap.h
--- a/Engine/Plugins/Lumin/MagicLeap/Source/MagicLeapAudio/Private/AudioMixerPlatformMagicLeap.h
+++ b/Engine/Plugins/Lumin/MagicLeap/Source/MagicLeapAudio/Private/AudioMixerPlatformMagicLeap.h
@@ -3,13 +3,8 @@
 #pragma once
 
 #include "AudioMixer.h"
-#include "MagicLeapPluginUtil.h" // for ML_INCLUDES_START/END
-
-#if WITH_MLSDK
-ML_INCLUDES_START
-#include <ml_audio.h>
-ML_INCLUDES_END
-#endif //WITH_MLSDK
+#include "Lumin/CAPIShims/LuminAPI.h"
+#include "FakeDeviceCallbackRunnable.h"
 
 // Any platform defines
 namespace Audio
@@ -44,6 +39,7 @@ namespace Audio
 		virtual FAudioPlatformSettings GetPlatformSettings() const override;
 		virtual void SuspendContext() override;
 		virtual void ResumeContext() override;
+		virtual void OnHardwareUpdate() override;
 
 		//~ End IAudioMixerPlatformInterface
 
@@ -52,6 +48,9 @@ namespace Audio
 
 		virtual int32 GetNumFrames(const int32 InNumReqestedFrames) override;
 
+		void DeviceStandby();
+		void DeviceActive();
+
 	private:
 #if WITH_MLSDK
 		static const TCHAR* GetErrorString(MLResult Result);
@@ -61,13 +60,14 @@ namespace Audio
 		bool bInitialized;
 		bool bInCallback;
 
+		FFakeDeviceCallbackRunnable FakeCallback;
+
 #if WITH_MLSDK
 		MLHandle StreamHandle;
-
+		FCriticalSection StreamHandleCriticalSection;
 		// Static callback used for MLAudio:
 		static void MLAudioCallback(MLHandle Handle, void* CallbackContext);
 #endif //WITH_MLSDK
 	};
 
 }
-
diff --git a/Engine/Plugins/Lumin/MagicLeap/Source/MagicLeapAudio/Private/FakeAudioMixerMagicLeap.cpp b/Engine/Plugins/Lumin/MagicLeap/Source/MagicLeapAudio/Private/FakeAudioMixerMagicLeap.cpp
new file mode 100644
--- /dev/null
+++ b/Engine/Plugins/Lumin/MagicLeap/Source/MagicLeapAudio/Private/FakeAudioMixerMagicLeap.cpp
@@ -0,0 +1,199 @@
+// Copyright 1998-2019 Epic Games, Inc. All Rights Reserved.
+
+#include "FakeAudioMixerMagicLeap.h"
+#include "Modules/ModuleManager.h"
+#include "AudioMixer.h"
+#include "AudioMixerDevice.h"
+#include "CoreGlobals.h"
+#include "VorbisAudioInfo.h"
+#include "ADPCMAudioInfo.h"
+#include "HeadMountedDisplayFunctionLibrary.h"
+
+namespace Audio
+{
+	// ML1 currently only has stereo speakers and stereo aux support.
+	const uint32 FakeDefaultNumChannels = 2;
+	// TODO: @Epic check the value to be used. Setting default for now.
+	const uint32 FakeDefaultSamplesPerSecond = 48000;  // presumed 48KHz and 16 bits for the sample
+
+	FFakeAudioMixerMagicLeap::FFakeAudioMixerMagicLeap()
+	: bInitialized(false)
+	, FakeCallback(this)
+	{}
+
+	FFakeAudioMixerMagicLeap::~FFakeAudioMixerMagicLeap()
+	{
+		TeardownHardware();
+	}
+
+	bool FFakeAudioMixerMagicLeap::InitializeHardware()
+	{
+		bInitialized = true;
+		return true;
+	}
+
+	bool FFakeAudioMixerMagicLeap::TeardownHardware()
+	{
+		bInitialized = false;
+		return true;
+	}
+
+	bool FFakeAudioMixerMagicLeap::IsInitialized() const
+	{
+		return bInitialized;
+	}
+
+	bool FFakeAudioMixerMagicLeap::GetNumOutputDevices(uint32& OutNumOutputDevices)
+	{
+		// ML1 will always have just one device.
+		OutNumOutputDevices = 1;
+		return true;
+	}
+
+	bool FFakeAudioMixerMagicLeap::GetOutputDeviceInfo(const uint32 InDeviceIndex, FAudioPlatformDeviceInfo& OutInfo)
+	{
+		OutInfo.Name = TEXT("Magic Leap Audio Device");
+		OutInfo.DeviceId = 0;
+		OutInfo.bIsSystemDefault = true;
+		OutInfo.NumChannels = FakeDefaultNumChannels;
+		OutInfo.OutputChannelArray.SetNum(2);
+		OutInfo.OutputChannelArray[0] = EAudioMixerChannel::FrontLeft;
+		OutInfo.OutputChannelArray[1] = EAudioMixerChannel::FrontRight;
+		OutInfo.SampleRate = FakeDefaultSamplesPerSecond;
+		OutInfo.Format = EAudioMixerStreamDataFormat::Float;
+
+		PlatformDeviceInfo = OutInfo;
+
+		return true;
+	}
+
+	bool FFakeAudioMixerMagicLeap::GetDefaultOutputDeviceIndex(uint32& OutDefaultDeviceIndex) const
+	{
+		OutDefaultDeviceIndex = 0;
+		return true;
+	}
+
+	bool FFakeAudioMixerMagicLeap::OpenAudioStream(const FAudioMixerOpenStreamParams& Params)
+	{
+		if (!bInitialized || AudioStreamInfo.StreamState != EAudioOutputStreamState::Closed)
+		{
+			return false;
+		}
+
+		if (!GetOutputDeviceInfo(AudioStreamInfo.OutputDeviceIndex, AudioStreamInfo.DeviceInfo))
+		{
+			return false;
+		}
+
+		OpenStreamParams = Params;
+		OpenStreamParams.NumFrames = 3200;
+
+		AudioStreamInfo.Reset();
+		AudioStreamInfo.OutputDeviceIndex = 0;
+		AudioStreamInfo.NumOutputFrames = OpenStreamParams.NumFrames;
+		AudioStreamInfo.NumBuffers = OpenStreamParams.NumBuffers;
+		AudioStreamInfo.AudioMixer = OpenStreamParams.AudioMixer;
+		GetOutputDeviceInfo(0, AudioStreamInfo.DeviceInfo);
+
+		AudioStreamInfo.StreamState = EAudioOutputStreamState::Open;
+
+		// Always use the fake callback with FFakeAudioMixerMagicLeap
+		FakeCallback.SetShouldUseCallback(true);
+
+		return true;
+	}
+
+	bool FFakeAudioMixerMagicLeap::CloseAudioStream()
+	{
+		if (!bInitialized || (AudioStreamInfo.StreamState != EAudioOutputStreamState::Open && AudioStreamInfo.StreamState != EAudioOutputStreamState::Stopped))
+		{
+			return false;
+		}
+
+		AudioStreamInfo.StreamState = EAudioOutputStreamState::Closed;
+		return true;
+	}
+
+	bool FFakeAudioMixerMagicLeap::StartAudioStream()
+	{
+		BeginGeneratingAudio();
+		return true;
+	}
+
+	bool FFakeAudioMixerMagicLeap::StopAudioStream()
+	{
+		if (!bInitialized || AudioStreamInfo.StreamState != EAudioOutputStreamState::Running)
+		{
+			return false;
+		}
+
+		if (AudioStreamInfo.StreamState == EAudioOutputStreamState::Running)
+		{
+			StopGeneratingAudio();
+		}
+
+		check(AudioStreamInfo.StreamState == EAudioOutputStreamState::Stopped);
+		return true;
+	}
+
+	FAudioPlatformDeviceInfo FFakeAudioMixerMagicLeap::GetPlatformDeviceInfo() const
+	{
+		return PlatformDeviceInfo;
+	}
+
+	FAudioPlatformSettings FFakeAudioMixerMagicLeap::GetPlatformSettings() const
+	{
+		FAudioPlatformSettings PlatformSettings;
+		PlatformSettings.CallbackBufferFrameSize = 3200;
+		PlatformSettings.MaxChannels = 0;
+		PlatformSettings.NumBuffers = 2;
+		PlatformSettings.SampleRate = 4800;
+		return PlatformSettings;
+	}
+
+	void FFakeAudioMixerMagicLeap::SubmitBuffer(const uint8* Buffer)
+	{}
+
+	FName FFakeAudioMixerMagicLeap::GetRuntimeFormat(USoundWave* InSoundWave)
+	{
+#if WITH_OGGVORBIS
+		static FName NAME_OGG(TEXT("OGG"));
+		if (InSoundWave->HasCompressedData(NAME_OGG))
+		{
+			return NAME_OGG;
+		}
+#endif
+
+		static FName NAME_ADPCM(TEXT("ADPCM"));
+
+		return NAME_ADPCM;
+	}
+
+	bool FFakeAudioMixerMagicLeap::HasCompressedAudioInfoClass(USoundWave* InSoundWave)
+	{
+		return true;
+	}
+
+	ICompressedAudioInfo* FFakeAudioMixerMagicLeap::CreateCompressedAudioInfo(USoundWave* InSoundWave)
+	{
+#if WITH_OGGVORBIS
+		static FName NAME_OGG(TEXT("OGG"));
+		if (InSoundWave->HasCompressedData(NAME_OGG))
+		{
+			return new FVorbisAudioInfo();
+		}
+#endif
+		static FName NAME_ADPCM(TEXT("ADPCM"));
+		return new FADPCMAudioInfo();
+	}
+
+	FString FFakeAudioMixerMagicLeap::GetDefaultDeviceName()
+	{
+		return FString(TEXT("MLAudio"));
+	}
+
+	int32 FFakeAudioMixerMagicLeap::GetNumFrames(const int32 InNumReqestedFrames)
+	{
+		return 3200;
+	}
+}
diff --git a/Engine/Plugins/Lumin/MagicLeap/Source/MagicLeapAudio/Private/FakeAudioMixerMagicLeap.h b/Engine/Plugins/Lumin/MagicLeap/Source/MagicLeapAudio/Private/FakeAudioMixerMagicLeap.h
new file mode 100644
--- /dev/null
+++ b/Engine/Plugins/Lumin/MagicLeap/Source/MagicLeapAudio/Private/FakeAudioMixerMagicLeap.h
@@ -0,0 +1,51 @@
+// Copyright 1998-2019 Epic Games, Inc. All Rights Reserved.
+
+#pragma once
+
+#include "AudioMixer.h"
+#include "FakeDeviceCallbackRunnable.h"
+
+namespace Audio
+{
+	// For an audio mixer module to be used in PIE, 
+	// it needs to have been successfully created on editor launch, even though a new instance is created for PIE.
+	// If MLRemote is not already running when the editor is launched, FMixerPlatformMagicLeap will cause the engine to crash.
+	// Thus, create the actual ML audio mixer only if an ML HMD is connected i.e. MLRemote is running.
+	// Otherwise, use a fake audio mixer so this module gets registered on editor startup, and works without crashing.
+	class FFakeAudioMixerMagicLeap : public IAudioMixerPlatformInterface
+	{
+
+	public:
+
+		FFakeAudioMixerMagicLeap();
+		~FFakeAudioMixerMagicLeap();
+
+		//~ Begin IAudioMixerPlatformInterface
+		virtual EAudioMixerPlatformApi::Type GetPlatformApi() const override { return EAudioMixerPlatformApi::Null; }
+		virtual bool InitializeHardware() override;
+		virtual bool TeardownHardware() override;
+		virtual bool IsInitialized() const override;
+		virtual bool GetNumOutputDevices(uint32& OutNumOutputDevices) override;
+		virtual bool GetOutputDeviceInfo(const uint32 InDeviceIndex, FAudioPlatformDeviceInfo& OutInfo) override;
+		virtual bool GetDefaultOutputDeviceIndex(uint32& OutDefaultDeviceIndex) const override;
+		virtual bool OpenAudioStream(const FAudioMixerOpenStreamParams& Params) override;
+		virtual bool CloseAudioStream() override;
+		virtual bool StartAudioStream() override;
+		virtual bool StopAudioStream() override;
+		virtual FAudioPlatformDeviceInfo GetPlatformDeviceInfo() const override;
+		virtual void SubmitBuffer(const uint8* Buffer) override;
+		virtual FName GetRuntimeFormat(USoundWave* InSoundWave) override;
+		virtual int32 GetNumFrames(const int32 InNumReqestedFrames) override;
+		virtual bool HasCompressedAudioInfoClass(USoundWave* InSoundWave) override;
+		virtual ICompressedAudioInfo* CreateCompressedAudioInfo(USoundWave* InSoundWave) override;
+		virtual FString GetDefaultDeviceName() override;
+		virtual FAudioPlatformSettings GetPlatformSettings() const override;
+
+		//~ End IAudioMixerPlatformInterface
+
+	private:
+		FAudioPlatformDeviceInfo PlatformDeviceInfo;
+		bool bInitialized;
+		FFakeDeviceCallbackRunnable FakeCallback;
+	};
+}
diff --git a/Engine/Plugins/Lumin/MagicLeap/Source/MagicLeapAudio/Private/FakeDeviceCallbackRunnable.cpp b/Engine/Plugins/Lumin/MagicLeap/Source/MagicLeapAudio/Private/FakeDeviceCallbackRunnable.cpp
new file mode 100644
--- /dev/null
+++ b/Engine/Plugins/Lumin/MagicLeap/Source/MagicLeapAudio/Private/FakeDeviceCallbackRunnable.cpp
@@ -0,0 +1,72 @@
+// Copyright 1998-2019 Epic Games, Inc. All Rights Reserved.
+
+#include "FakeDeviceCallbackRunnable.h"
+#include "CoreGlobals.h"
+
+namespace Audio
+{
+	FFakeDeviceCallbackRunnable::FFakeDeviceCallbackRunnable(IAudioMixerPlatformInterface* InContext)
+	: Context(InContext)
+	, Thread(nullptr)
+	, StopTaskCounter(0)
+	, Semaphore(nullptr)
+	, bShouldUseCallback(false)
+	{
+	}
+
+	FFakeDeviceCallbackRunnable::~FFakeDeviceCallbackRunnable()
+	{
+		StopTaskCounter.Increment();
+
+		if (Semaphore)
+		{
+			Semaphore->Trigger();
+			Thread->WaitForCompletion();
+			FGenericPlatformProcess::ReturnSynchEventToPool(Semaphore);
+			Semaphore = nullptr;
+		}
+
+		delete Thread;
+		Thread = nullptr;
+	}
+
+	uint32 FFakeDeviceCallbackRunnable::Run()
+	{
+		while (StopTaskCounter.GetValue() == 0)
+		{
+			if (bShouldUseCallback == false)
+			{
+				Semaphore->Wait();
+			}
+
+			Context->ReadNextBuffer();
+			Semaphore->Wait(16);
+		}
+
+		return 0;
+	}
+
+	void FFakeDeviceCallbackRunnable::SetShouldUseCallback(bool bUseCallback)
+	{
+		if (bShouldUseCallback != bUseCallback)
+		{
+			bShouldUseCallback = bUseCallback;
+
+			if (bShouldUseCallback)
+			{
+				WakeUp();
+			}
+		}
+	}
+
+	void FFakeDeviceCallbackRunnable::WakeUp()
+	{
+		if (Semaphore == nullptr)
+		{
+			Semaphore = FGenericPlatformProcess::GetSynchEventFromPool();
+			Thread = FRunnableThread::Create(this, TEXT("FakeDeviceCallbackRunnable"), 0, TPri_BelowNormal);
+		}
+
+		Semaphore->Trigger();
+	}
+}
diff --git a/Engine/Plugins/Lumin/MagicLeap/Source/MagicLeapAudio/Private/FakeDeviceCallbackRunnable.h b/Engine/Plugins/Lumin/MagicLeap/Source/MagicLeapAudio/Private/FakeDeviceCallbackRunnable.h
new file mode 100644
--- /dev/null
+++ b/Engine/Plugins/Lumin/MagicLeap/Source/MagicLeapAudio/Private/FakeDeviceCallbackRunnable.h
@@ -0,0 +1,41 @@
+// Copyright 1998-2019 Epic Games, Inc. All Rights Reserved.
+
+#pragma once
+
+#include "AudioMixer.h"
+#include "Containers/Ticker.h"
+#include "HAL/Runnable.h"
+#include "HAL/RunnableThread.h"
+#include "HAL/ThreadSafeBool.h"
+#include "HAL/ThreadSafeCounter64.h"
+
+namespace Audio
+{
+	// When the device goes in standby mode, the MLAudioCallback function is not called.
+	// As a result the audio buffers from the engine queue up and they play back once the device becomes active again.
+	// This desyncs gameplay with audio and breaks the UX spec for standby mode. When in standby, the app should function normally
+	// from a gameplay standpoint. It should not "pause". This spec is to ensure that when the device is active again, the "resume" time is very minimal.
+	// Using this fake callback, we keep emptying the audio buffer from the engine and act as if everything is running like it normally should.
+	// Another aspect this fixes is launching an app when device is already in standby mode. If during engine initialization, the audio mixer
+	// does not call IAudioMixerPlatformInterface::ReadNextBuffer(), the audio engine blocks and eventually crashes.
+	// This fake callback ensures a smooth initialization as well.
+	class FFakeDeviceCallbackRunnable : public FRunnable
+	{
+	public:
+		FFakeDeviceCallbackRunnable(class IAudioMixerPlatformInterface* InContext);
+		virtual ~FFakeDeviceCallbackRunnable();
+		virtual uint32 Run() override;
+
+		void SetShouldUseCallback(bool bUseCallback);
+
+	private:
+		void WakeUp();
+
+		class IAudioMixerPlatformInterface* Context;
+		FRunnableThread* Thread;
+		FThreadSafeCounter StopTaskCounter;
+		FEvent* Semaphore;
+
+		FThreadSafeBool bShouldUseCallback;
+	};
+}
diff --git a/Engine/Plugins/Lumin/MagicLeap/Source/MagicLeapAudio/Private/MagicLeapAudioFunctionLibrary.cpp b/Engine/Plugins/Lumin/MagicLeap/Source/MagicLeapAudio/Private/MagicLeapAudioFunctionLibrary.cpp
new file mode 100644
--- /dev/null
+++ b/Engine/Plugins/Lumin/MagicLeap/Source/MagicLeapAudio/Private/MagicLeapAudioFunctionLibrary.cpp
@@ -0,0 +1,18 @@
+// Copyright 1998-2019 Epic Games, Inc. All Rights Reserved.
+
+#include "MagicLeapAudioFunctionLibrary.h"
+#include "MagicLeapAudioModule.h"
+
+bool UMagicLeapAudioFunctionLibrary::SetOnAudioJackPluggedDelegate(const FMagicLeapAudioJackPluggedDelegate& InResultDelegate)
+{
+	FMagicLeapAudioJackPluggedDelegateMulti ResultDelegate;
+	ResultDelegate.Add(InResultDelegate);
+	return GetMagicLeapAudioModule().SetOnAudioJackPluggedDelegate(ResultDelegate);
+}
+
+bool UMagicLeapAudioFunctionLibrary::SetOnAudioJackUnpluggedDelegate(const FMagicLeapAudioJackUnpluggedDelegate& InResultDelegate)
+{
+	FMagicLeapAudioJackUnpluggedDelegateMulti ResultDelegate;
+	ResultDelegate.Add(InResultDelegate);
+	return GetMagicLeapAudioModule().SetOnAudioJackUnpluggedDelegate(ResultDelegate);
+}
diff --git a/Engine/Plugins/Lumin/MagicLeap/Source/MagicLeapAudio/Private/MagicLeapAudioModule.cpp b/Engine/Plugins/Lumin/MagicLeap/Source/MagicLeapAudio/Private/MagicLeapAudioModule.cpp
new file mode 100644
--- /dev/null
+++ b/Engine/Plugins/Lumin/MagicLeap/Source/MagicLeapAudio/Private/MagicLeapAudioModule.cpp
@@ -0,0 +1,88 @@
+// Copyright 1998-2019 Epic Games, Inc. All Rights Reserved.
+
+#include "MagicLeapAudioModule.h"
+#include "Engine/Engine.h"
+
+DEFINE_LOG_CATEGORY(LogMagicLeapAudio);
+
+FMagicLeapAudioModule::FMagicLeapAudioModule()
+: PrevOutputDevice(EOutputDevice::Undertermined)
+{
+}
+
+FAudioDevice* FMagicLeapAudioModule::CreateAudioDevice()
+{
+#if PLATFORM_WINDOWS || PLATFORM_MAC
+	const IMagicLeapPlugin& MLPlugin = IMagicLeapPlugin::Get();
+	if (MLPlugin.IsMagicLeapHMDValid() && MLPlugin.UseMLAudioForZI())
+	{
+		// For an audio mixer module to be used in PIE, 
+		// it needs to have been successfully created on editor launch, even though a new instance is created for PIE.
+		// If MLRemote is not already running when the editor is launched, FMixerPlatformMagicLeap will cause the engine to crash.
+		// Thus, create the actual ML audio mixer only if an ML HMD is connected i.e. MLRemote is running.
+		// Otherwise, use a fake audio mixer so this module gets registered on editor startup, and works without crashing.
+		if (UHeadMountedDisplayFunctionLibrary::IsHeadMountedDisplayConnected())
+		{
+			return new Audio::FMixerDevice(new Audio::FMixerPlatformMagicLeap());
+		}
+		return new Audio::FMixerDevice(new Audio::FFakeAudioMixerMagicLeap());;
+	}
+
+	return nullptr;
+#elif PLATFORM_LUMIN
+	return new Audio::FMixerDevice(new Audio::FMixerPlatformMagicLeap());
+#else
+	return nullptr;
+#endif // PLATFORM_LUMIN
+}
+
+bool FMagicLeapAudioModule::SetOnAudioJackPluggedDelegate(const FMagicLeapAudioJackPluggedDelegateMulti& ResultDelegate)
+{
+	OnAudioJackPluggedDelegateMulti = ResultDelegate;
+	return true;
+}
+
+bool FMagicLeapAudioModule::SetOnAudioJackUnpluggedDelegate(const FMagicLeapAudioJackUnpluggedDelegateMulti& ResultDelegate)
+{
+	OnAudioJackUnpluggedDelegateMulti = ResultDelegate;
+	return true;
+}
+
+bool FMagicLeapAudioModule::CheckOutputDevice()
+{
+#if WITH_MLSDK
+	// Output device changes are in human time, and MLAudioGetOutputDevice is an expensive call
+	static constexpr double CheckOutputDeviceRate = 1.0;
+	const auto CurrentTime = FPlatformTime::Seconds();
+	if ((CurrentTime - LastOutputDeviceCheckedTime) >= CheckOutputDeviceRate)
+	{
+		LastOutputDeviceCheckedTime = CurrentTime;
+		MLAudioOutputDevice CurrOutputDevice;
+		MLResult Result = MLAudioGetOutputDevice(&CurrOutputDevice);
+		if (Result != MLResult_Ok)
+		{
+			UE_LOG(LogMagicLeapAudio, Error, TEXT("MLAudioGetOutputDevice failed with error '%s'."), UTF8_TO_TCHAR(MLAudioGetResultString(Result)));
+			return false;
+		}
+
+		if (CurrOutputDevice == MLAudioOutputDevice_Wearable)
+		{
+			if (OnAudioJackUnpluggedDelegateMulti.IsBound() && PrevOutputDevice != EOutputDevice::Wearable)
+			{
+				PrevOutputDevice = EOutputDevice::Wearable;
+				OnAudioJackUnpluggedDelegateMulti.Broadcast();
+			}
+		}
+		else if (OnAudioJackPluggedDelegateMulti.IsBound() && PrevOutputDevice != EOutputDevice::AudioJack)
+		{
+			PrevOutputDevice = EOutputDevice::AudioJack;
+			OnAudioJackPluggedDelegateMulti.Broadcast();
+		}
+	}
+	return true;
+#else
+	return false;
+#endif // WITH_MLSDK
+}
+
+IMPLEMENT_MODULE(FMagicLeapAudioModule, MagicLeapAudio);
diff --git a/Engine/Plugins/Lumin/MagicLeap/Source/MagicLeapAudio/Private/MagicLeapAudioModule.h b/Engine/Plugins/Lumin/MagicLeap/Source/MagicLeapAudio/Private/MagicLeapAudioModule.h
new file mode 100644
--- /dev/null
+++ b/Engine/Plugins/Lumin/MagicLeap/Source/MagicLeapAudio/Private/MagicLeapAudioModule.h
@@ -0,0 +1,42 @@
+// Copyright 1998-2019 Epic Games, Inc. All Rights Reserved.
+
+#pragma once
+
+#include "CoreMinimal.h"
+#include "IMagicLeapPlugin.h"
+#include "IMagicLeapAudioModule.h"
+#include "MagicLeapAudioTypes.h"
+#include "AudioMixer.h"
+#include "AudioMixerPlatformMagicLeap.h"
+#include "FakeAudioMixerMagicLeap.h"
+#include "HeadMountedDisplayFunctionLibrary.h"
+#include "Lumin/CAPIShims/LuminAPIAudio.h"
+
+DECLARE_LOG_CATEGORY_EXTERN(LogMagicLeapAudio, Verbose, All);
+
+class FMagicLeapAudioModule : public IMagicLeapAudioModule
+{
+public:
+	FMagicLeapAudioModule();
+	FAudioDevice* CreateAudioDevice() override;
+	bool SetOnAudioJackPluggedDelegate(const FMagicLeapAudioJackPluggedDelegateMulti& ResultDelegate);
+	bool SetOnAudioJackUnpluggedDelegate(const FMagicLeapAudioJackUnpluggedDelegateMulti& ResultDelegate);
+	bool CheckOutputDevice();
+
+private:
+	enum EOutputDevice
+	{
+		Undertermined,
+		Wearable,
+		AudioJack,
+	};
+	EOutputDevice PrevOutputDevice;
+	FMagicLeapAudioJackPluggedDelegateMulti OnAudioJackPluggedDelegateMulti;
+	FMagicLeapAudioJackUnpluggedDelegateMulti OnAudioJackUnpluggedDelegateMulti;
+	double LastOutputDeviceCheckedTime = 0.0;
+};
+
+inline FMagicLeapAudioModule& GetMagicLeapAudioModule()
+{
+	return FModuleManager::Get().GetModuleChecked<FMagicLeapAudioModule>("MagicLeapAudio");
+}
diff --git a/Engine/Plugins/Lumin/MagicLeap/Source/MagicLeapAudio/Public/IMagicLeapAudioModule.h b/Engine/Plugins/Lumin/MagicLeap/Source/MagicLeapAudio/Public/IMagicLeapAudioModule.h
new file mode 100644
--- /dev/null
+++ b/Engine/Plugins/Lumin/MagicLeap/Source/MagicLeapAudio/Public/IMagicLeapAudioModule.h
@@ -0,0 +1,43 @@
+// Copyright 1998-2019 Epic Games, Inc. All Rights Reserved.
+
+#pragma once
+
+#include "Modules/ModuleManager.h"
+#include "AudioMixerDevice.h"
+
+	/**
+	 * The public interface to this module.  In most cases, this interface is only public to sibling modules
+	 * within this plugin.
+	 */
+class IMagicLeapAudioModule : public IAudioDeviceModule
+{
+public:
+	/**
+	 * Singleton-like access to this module's interface.  This is just for convenience!
+	 * Beware of calling this during the shutdown phase, though.  Your module might have been unloaded already.
+	 *
+	 * @return Returns singleton instance, loading the module on demand if needed
+	 */
+	static inline IMagicLeapAudioModule& Get()
+	{
+		return FModuleManager::LoadModuleChecked<IMagicLeapAudioModule>("MagicLeapAudio");
+	}
+
+	/**
+	 * Checks to see if this module is loaded and ready.  It is only valid to call Get() if IsAvailable() returns true.
+	 *
+	 * @return True if the module is loaded and ready to use
+	 */
+	static inline bool IsAvailable()
+	{
+		return FModuleManager::Get().IsModuleLoaded("MagicLeapAudio");
+	}
+
+	/**
+	 * Creates a new instance of the audio device implemented by the module.
+	 * @note Overriden by FMagicLeapAudioModule.
+	 *
+	 * @return A new instance of the audio device.
+	 */
+	FAudioDevice* CreateAudioDevice() override { return nullptr; }
+};
diff --git a/Engine/Plugins/Lumin/MagicLeap/Source/MagicLeapAudio/Public/MagicLeapAudioFunctionLibrary.h b/Engine/Plugins/Lumin/MagicLeap/Source/MagicLeapAudio/Public/MagicLeapAudioFunctionLibrary.h
new file mode 100644
--- /dev/null
+++ b/Engine/Plugins/Lumin/MagicLeap/Source/MagicLeapAudio/Public/MagicLeapAudioFunctionLibrary.h
@@ -0,0 +1,28 @@
+// Copyright 1998-2019 Epic Games, Inc. All Rights Reserved.
+#pragma once
+
+#include "Engine/Engine.h"
+#include "Kismet/BlueprintFunctionLibrary.h"
+#include "MagicLeapAudioTypes.h"
+#include "MagicLeapAudioFunctionLibrary.generated.h"
+
+UCLASS(ClassGroup = MagicLeap)
+class MAGICLEAPAUDIO_API UMagicLeapAudioFunctionLibrary : public UBlueprintFunctionLibrary
+{
+	GENERATED_BODY()
+
+public:
+	/** 
+		Sets the delegate used to notify that an audio device has been plugged into the audio jack.
+		@param ResultDelegate The delegate that will be notified when an audio device has been plugged into the audio jack.
+	*/
+	UFUNCTION(BlueprintCallable, Category = "Audio Function Library | MagicLeap")
+	static bool SetOnAudioJackPluggedDelegate(const FMagicLeapAudioJackPluggedDelegate& ResultDelegate);
+
+	/**
+		Sets the delegate used to notify that an audio device has been unplugged from the audio jack.
+		@param ResultDelegate The delegate that will be notified when an audio device has been unplugged from the audio jack.
+	*/
+	UFUNCTION(BlueprintCallable, Category = "Audio Function Library | MagicLeap")
+	static bool SetOnAudioJackUnpluggedDelegate(const FMagicLeapAudioJackUnpluggedDelegate& ResultDelegate);
+};
diff --git a/Engine/Plugins/Lumin/MagicLeap/Source/MagicLeapAudio/Public/MagicLeapAudioTypes.h b/Engine/Plugins/Lumin/MagicLeap/Source/MagicLeapAudio/Public/MagicLeapAudioTypes.h
new file mode 100644
--- /dev/null
+++ b/Engine/Plugins/Lumin/MagicLeap/Source/MagicLeapAudio/Public/MagicLeapAudioTypes.h
@@ -0,0 +1,16 @@
+// Copyright 1998-2019 Epic Games, Inc. All Rights Reserved.
+#pragma once
+
+#include "Engine/Engine.h"
+#include "MagicLeapAudioTypes.generated.h"
+
+USTRUCT()
+struct FMagicLeapAudioDummyStruct { GENERATED_BODY() };
+
+/** Delegate used to notify that an audio device has been plugged into the audio jack. */
+DECLARE_DYNAMIC_DELEGATE(FMagicLeapAudioJackPluggedDelegate);
+DECLARE_DYNAMIC_MULTICAST_DELEGATE(FMagicLeapAudioJackPluggedDelegateMulti);
+
+/** Delegate used to notify that an audio device has been unplugged from the audio jack. */
+DECLARE_DYNAMIC_DELEGATE(FMagicLeapAudioJackUnpluggedDelegate);
+DECLARE_DYNAMIC_MULTICAST_DELEGATE(FMagicLeapAudioJackUnpluggedDelegateMulti);
diff --git a/Engine/Plugins/Lumin/MagicLeap/Source/MagicLeapAudioCapture/Private/AudioCaptureMagicLeap.h b/Engine/Plugins/Lumin/MagicLeap/Source/MagicLeapAudioCapture/Private/AudioCaptureMagicLeap.h
--- a/Engine/Plugins/Lumin/MagicLeap/Source/MagicLeapAudioCapture/Private/AudioCaptureMagicLeap.h
+++ b/Engine/Plugins/Lumin/MagicLeap/Source/MagicLeapAudioCapture/Private/AudioCaptureMagicLeap.h
@@ -2,13 +2,7 @@
 #pragma once
 
 #include "AudioCaptureCore.h"
-#include "MagicLeapPluginUtil.h"
-
-#if WITH_MLSDK
-ML_INCLUDES_START
-#include <ml_audio.h>
-ML_INCLUDES_END
-#endif //WITH_MLSDK
+#include "Lumin/CAPIShims/LuminAPIAudio.h"
 
 namespace Audio
 {
diff --git a/Engine/Plugins/Lumin/MagicLeap/Source/MagicLeapController/MagicLeapController.Build.cs b/Engine/Plugins/Lumin/MagicLeap/Source/MagicLeapController/MagicLeapController.Build.cs
--- a/Engine/Plugins/Lumin/MagicLeap/Source/MagicLeapController/MagicLeapController.Build.cs
+++ b/Engine/Plugins/Lumin/MagicLeap/Source/MagicLeapController/MagicLeapController.Build.cs
@@ -32,9 +32,6 @@ namespace UnrealBuildTool.Rules
 					"MLSDK"
 				}
 			);
-
-			// This is not ideal but needs to be done in order to expose the private MagicLeapHMD header to this module.
-			PrivateIncludePaths.Add(Path.Combine(new string[] { ModuleDirectory, "..", "MagicLeap", "Private" }));
 		}
 	}
 }
diff --git a/Engine/Plugins/Lumin/MagicLeap/Source/MagicLeapController/Private/MagicLeapController.cpp b/Engine/Plugins/Lumin/MagicLeap/Source/MagicLeapController/Private/MagicLeapController.cpp
--- a/Engine/Plugins/Lumin/MagicLeap/Source/MagicLeapController/Private/MagicLeapController.cpp
+++ b/Engine/Plugins/Lumin/MagicLeap/Source/MagicLeapController/Private/MagicLeapController.cpp
@@ -2,7 +2,6 @@
 
 #include "MagicLeapController.h"
 #include "IMagicLeapPlugin.h"
-#include "MagicLeapHMD.h"
 #include "MagicLeapMath.h"
 #include "HAL/PlatformProcess.h"
 #include "Misc/App.h"
@@ -10,7 +9,6 @@
 #include "MagicLeapControllerKeys.h"
 #include "Framework/Application/SlateApplication.h"
 #include "Async/Async.h"
-#include "AppFramework.h"
 #include "MagicLeapPluginUtil.h"
 #include "TouchpadGesturesComponent.h"
 #include "AssetData.h"
@@ -23,8 +21,8 @@
 #include "MagicLeapControllerMappings.h"
 
 #if WITH_MLSDK
-static_assert(MLInput_MaxControllerTouchpadTouches == FMagicLeapControllerState::kMaxTouches, "Mismatch in max touch constants");
 
+static_assert(MLInput_MaxControllerTouchpadTouches == FMagicLeapControllerState::kMaxTouches, "Mismatch in max touch constants");
 // Controller Mapper
 FMagicLeapController::FControllerMapper::FControllerMapper()
 {
@@ -270,25 +268,25 @@ uint8 FMagicLeapController::FControllerMapper::GetInputControllerIndexForHand(EC
 	return 0xFF;
 }
 
-EMLControllerType FMagicLeapController::FControllerMapper::MotionSourceToControllerType(FName MotionSource)
+EMagicLeapControllerType FMagicLeapController::FControllerMapper::MotionSourceToControllerType(FName MotionSource)
 {
 	auto MLSourceToControllerType = [](FName InMotionSource)
 	{
 		if (InMotionSource == FMagicLeapMotionSourceNames::Control0 ||
 			InMotionSource == FMagicLeapMotionSourceNames::Control1)
 		{
-			return EMLControllerType::Device;
+			return EMagicLeapControllerType::Device;
 		}
 		if (InMotionSource == FMagicLeapMotionSourceNames::MobileApp)
 		{
-			return EMLControllerType::MobileApp;
+			return EMagicLeapControllerType::MobileApp;
 		}
-		return EMLControllerType::None;
+		return EMagicLeapControllerType::None;
 	};
 
 	// First just see if it's one of ours and can be easily mapped
 	auto ControllerType = MLSourceToControllerType(MotionSource);
-	if (ControllerType == EMLControllerType::None)
+	if (ControllerType == EMagicLeapControllerType::None)
 	{
 		// If not, see if it's a hand mapping
 		auto ControllerHand = EControllerHand::ControllerHand_Count;
@@ -321,7 +319,7 @@ FMagicLeapController::FMagicLeapController(const TSharedRef<FGenericApplicationM
 	, InputTracker(ML_INVALID_HANDLE)
 	, ControllerTracker(ML_INVALID_HANDLE)
 	, ControllerDof(MLInputControllerDof_6)
-	, TrackingMode(EMLControllerTrackingMode::CoordinateFrameUID)
+	, TrackingMode(EMagicLeapControllerTrackingMode::CoordinateFrameUID)
 #endif //WITH_MLSDK
 	, bIsInputStateValid(false)
 	, TriggerKeyIsConsideredPressed(80.0f)
@@ -348,7 +346,7 @@ FMagicLeapController::FMagicLeapController(const TSharedRef<FGenericApplicationM
 
 	// We're implicitly requiring that the MagicLeapPlugin has been loaded and
 	// initialized at this point.
-	IMagicLeapPlugin::Get().RegisterMagicLeapInputDevice(this);
+	IMagicLeapPlugin::Get().RegisterMagicLeapTrackerEntity(this);
 }
 
 FMagicLeapController::~FMagicLeapController()
@@ -357,14 +355,109 @@ FMagicLeapController::~FMagicLeapController()
 	// but it isn't an assumption that we should make.
 	if (IMagicLeapPlugin::IsAvailable())
 	{
-		IMagicLeapPlugin::Get().UnregisterMagicLeapInputDevice(this);
+		IMagicLeapPlugin::Get().UnregisterMagicLeapTrackerEntity(this);
 	}
 
-	Disable();
+	DestroyEntityTracker();
 
 	IModularFeatures::Get().UnregisterModularFeature(GetModularFeatureName(), this);
 }
 
+#if WITH_MLSDK
+FName GetKeyNameFromGesture(const MLInputControllerTouchpadGesture* Gesture)
+{
+	FName KeyName = NAME_None;
+
+	switch(Gesture->type)
+	{
+		case MLInputControllerTouchpadGestureType::MLInputControllerTouchpadGestureType_Tap:
+			KeyName = FMagicLeapControllerKeyNames::TouchpadGesture_Tap_Name;
+			break;
+		case MLInputControllerTouchpadGestureType::MLInputControllerTouchpadGestureType_ForceTapDown:
+			KeyName = FMagicLeapControllerKeyNames::TouchpadGesture_ForceTapDown_Name;
+			break;
+		case MLInputControllerTouchpadGestureType::MLInputControllerTouchpadGestureType_ForceTapUp:
+			KeyName = FMagicLeapControllerKeyNames::TouchpadGesture_ForceTapUp_Name;
+			break;
+		case MLInputControllerTouchpadGestureType::MLInputControllerTouchpadGestureType_ForceDwell:
+			KeyName = FMagicLeapControllerKeyNames::TouchpadGesture_ForceDwell_Name;
+			break;
+		case MLInputControllerTouchpadGestureType::MLInputControllerTouchpadGestureType_SecondForceDown:
+			KeyName = FMagicLeapControllerKeyNames::TouchpadGesture_SecondForceDown_Name;
+			break;
+		case MLInputControllerTouchpadGestureType::MLInputControllerTouchpadGestureType_LongHold:
+			KeyName = FMagicLeapControllerKeyNames::TouchpadGesture_LongHold_Name;
+			break;
+		case MLInputControllerTouchpadGestureType::MLInputControllerTouchpadGestureType_RadialScroll:
+		{
+			switch(Gesture->direction)
+			{
+				case MLInputControllerTouchpadGestureDirection::MLInputControllerTouchpadGestureDirection_Clockwise:
+					KeyName = FMagicLeapControllerKeyNames::TouchpadGesture_RadialScroll_Clockwise_Name;
+					break;
+				case MLInputControllerTouchpadGestureDirection::MLInputControllerTouchpadGestureDirection_CounterClockwise:
+					KeyName = FMagicLeapControllerKeyNames::TouchpadGesture_RadialScroll_CounterClockwise_Name;
+					break;
+			}
+		}
+		break;
+		case MLInputControllerTouchpadGestureType::MLInputControllerTouchpadGestureType_Swipe:
+		{
+			switch(Gesture->direction)
+			{
+				case MLInputControllerTouchpadGestureDirection::MLInputControllerTouchpadGestureDirection_Up:
+					KeyName = FMagicLeapControllerKeyNames::TouchpadGesture_Swipe_Up_Name;
+					break;
+				case MLInputControllerTouchpadGestureDirection::MLInputControllerTouchpadGestureDirection_Down:
+					KeyName = FMagicLeapControllerKeyNames::TouchpadGesture_Swipe_Down_Name;
+					break;
+				case MLInputControllerTouchpadGestureDirection::MLInputControllerTouchpadGestureDirection_Right:
+					KeyName = FMagicLeapControllerKeyNames::TouchpadGesture_Swipe_Right_Name;
+					break;
+				case MLInputControllerTouchpadGestureDirection::MLInputControllerTouchpadGestureDirection_Left:
+					KeyName = FMagicLeapControllerKeyNames::TouchpadGesture_Swipe_Left_Name;
+					break;
+			}
+		}
+		break;
+		case MLInputControllerTouchpadGestureType::MLInputControllerTouchpadGestureType_Scroll:
+		{
+			switch(Gesture->direction)
+			{
+				case MLInputControllerTouchpadGestureDirection::MLInputControllerTouchpadGestureDirection_Up:
+					KeyName = FMagicLeapControllerKeyNames::TouchpadGesture_Scroll_Up_Name;
+					break;
+				case MLInputControllerTouchpadGestureDirection::MLInputControllerTouchpadGestureDirection_Down:
+					KeyName = FMagicLeapControllerKeyNames::TouchpadGesture_Scroll_Down_Name;
+					break;
+				case MLInputControllerTouchpadGestureDirection::MLInputControllerTouchpadGestureDirection_Right:
+					KeyName = FMagicLeapControllerKeyNames::TouchpadGesture_Scroll_Right_Name;
+					break;
+				case MLInputControllerTouchpadGestureDirection::MLInputControllerTouchpadGestureDirection_Left:
+					KeyName = FMagicLeapControllerKeyNames::TouchpadGesture_Scroll_Left_Name;
+					break;
+			}
+		}
+		break;
+		case MLInputControllerTouchpadGestureType::MLInputControllerTouchpadGestureType_Pinch:
+		{
+			switch(Gesture->direction)
+			{
+				case MLInputControllerTouchpadGestureDirection::MLInputControllerTouchpadGestureDirection_In:
+					KeyName = FMagicLeapControllerKeyNames::TouchpadGesture_Pinch_In_Name;
+					break;
+				case MLInputControllerTouchpadGestureDirection::MLInputControllerTouchpadGestureDirection_Out:
+					KeyName = FMagicLeapControllerKeyNames::TouchpadGesture_Pinch_Out_Name;
+					break;
+			}
+		}
+		break;
+	}
+
+	return KeyName;
+}
+#endif // WITH_MLSDK
+
 void FMagicLeapController::InitializeInputCallbacks()
 {
 #if WITH_MLSDK
@@ -378,10 +471,10 @@ void FMagicLeapController::InitializeInputCallbacks()
 			const auto& Hand = Controller->ControllerMapper.GetHandForInputControllerIndex(controller_id);
 			if (Hand != EControllerHand::ControllerHand_Count)
 			{
-				auto Gesture = MLToUnrealTouchpadGesture(Hand, FMagicLeapMotionSourceNames::Unknown, *touchpad_gesture);
+				FMagicLeapTouchpadGesture Gesture_DEPRECATED = MLToUnrealTouchpadGesture_DEPRECATED(Hand, FMagicLeapMotionSourceNames::Unknown, *touchpad_gesture);
 				for (IMagicLeapTouchpadGestures* Receiver : Controller->TouchpadGestureReceivers)
 				{
-					Receiver->OnTouchpadGestureStartCallback(Gesture);
+					Receiver->OnTouchpadGestureStartCallback(Gesture_DEPRECATED);
 				}
 			}
 		}
@@ -395,10 +488,10 @@ void FMagicLeapController::InitializeInputCallbacks()
 			const auto& Hand = Controller->ControllerMapper.GetHandForInputControllerIndex(controller_id);
 			if (Hand != EControllerHand::ControllerHand_Count)
 			{
-				auto Gesture = MLToUnrealTouchpadGesture(Hand, FMagicLeapMotionSourceNames::Unknown, *touchpad_gesture);
+				FMagicLeapTouchpadGesture Gesture_DEPRECATED = MLToUnrealTouchpadGesture_DEPRECATED(Hand, FMagicLeapMotionSourceNames::Unknown, *touchpad_gesture);
 				for (IMagicLeapTouchpadGestures* Receiver : Controller->TouchpadGestureReceivers)
 				{
-					Receiver->OnTouchpadGestureContinueCallback(Gesture);
+					Receiver->OnTouchpadGestureContinueCallback(Gesture_DEPRECATED);
 				}
 			}
 		}
@@ -412,10 +505,10 @@ void FMagicLeapController::InitializeInputCallbacks()
 			const auto& Hand = Controller->ControllerMapper.GetHandForInputControllerIndex(controller_id);
 			if (Hand != EControllerHand::ControllerHand_Count)
 			{
-				auto Gesture = MLToUnrealTouchpadGesture(Hand, FMagicLeapMotionSourceNames::Unknown, *touchpad_gesture);
+				FMagicLeapTouchpadGesture Gesture_DEPRECATED = MLToUnrealTouchpadGesture_DEPRECATED(Hand, FMagicLeapMotionSourceNames::Unknown, *touchpad_gesture);
 				for (IMagicLeapTouchpadGestures* Receiver : Controller->TouchpadGestureReceivers)
 				{
-					Receiver->OnTouchpadGestureEndCallback(Gesture);
+					Receiver->OnTouchpadGestureEndCallback(Gesture_DEPRECATED);
 				}
 			}
 		}
@@ -459,7 +552,7 @@ void FMagicLeapController::SendControllerEvents()
 #if WITH_MLSDK
 	if (bIsInputStateValid && MessageHandler.IsValid())
 	{
-		FMagicLeapHMD::EnableInput EnableInputFromHMD;
+		MagicLeap::EnableInput EnableInputFromHMD;
 		// fixes unreferenced parameter error for Windows package builds.
 		(void)EnableInputFromHMD;
 
@@ -510,7 +603,7 @@ void FMagicLeapController::ReadConfigParams()
 #if WITH_MLSDK
 	// Pull hand-mapping preferences from config file. If there are none, the default (legacy)
 	// mapping of device 0 to right and device 1 to left will persist.
-	const UEnum* ControllerHandEnum = StaticEnum<EControllerHand>();
+	const UEnum* ControllerHandEnum = FindObject<UEnum>(ANY_PACKAGE, TEXT("EControllerHand"));
 	TArray<FString> ControllerHands;
 	GConfig->GetArray(TEXT("/Script/LuminRuntimeSettings.LuminRuntimeSettings"),
 		TEXT("ControllerHands"), ControllerHands, GEngineIni);
@@ -555,7 +648,7 @@ void FMagicLeapController::ReadConfigParams()
 		TEXT("ControllerTrackingType"), ConfigString, GEngineIni);
 	if (ConfigString.Len() > 0)
 	{
-		const UEnum* TrackingTypeEnum = StaticEnum<ETrackingStatus>();
+		const UEnum* TrackingTypeEnum = FindObject<UEnum>(ANY_PACKAGE, TEXT("ETrackingStatus"));
 
 		auto TrackingTypeIndex = TrackingTypeEnum->GetValueByNameString(ConfigString);
 		if (TrackingTypeIndex != INDEX_NONE)
@@ -582,17 +675,17 @@ void FMagicLeapController::ReadConfigParams()
 		}
 	}
 
-	const static UEnum* TrackingModeEnum = StaticEnum<EMLControllerTrackingMode>();
+	const static UEnum* TrackingModeEnum = FindObject<UEnum>(ANY_PACKAGE, TEXT("EMagicLeapControllerTrackingMode"));
 	GConfig->GetString(TEXT("/Script/LuminRuntimeSettings.LuminRuntimeSettings"),
 		TEXT("ControllerTrackingMode"), ConfigString, GEngineIni);
 	if (ConfigString.Len() > 0)
 	{
-		TrackingMode = static_cast<EMLControllerTrackingMode>(TrackingModeEnum->GetValueByNameString(ConfigString));
+		TrackingMode = static_cast<EMagicLeapControllerTrackingMode>(TrackingModeEnum->GetValueByNameString(ConfigString));
 	}
 #endif //WITH_MLSDK
 }
 
-void FMagicLeapController::Enable()
+void FMagicLeapController::CreateEntityTracker()
 {
 #if WITH_MLSDK
 	MLResult Result = MLResult_Ok;
@@ -628,7 +721,7 @@ void FMagicLeapController::Enable()
 		ControllerTracker = ML_INVALID_HANDLE;
 
 		// Revert to input
-		TrackingMode = EMLControllerTrackingMode::InputService;
+		TrackingMode = EMagicLeapControllerTrackingMode::InputService;
 	}
 
 #if PLATFORM_LUMIN
@@ -669,12 +762,7 @@ void FMagicLeapController::Enable()
 #endif //WITH_MLSDK
 }
 
-bool FMagicLeapController::SupportsExplicitEnable() const
-{
-	return true;
-}
-
-void FMagicLeapController::Disable()
+void FMagicLeapController::DestroyEntityTracker()
 {
 #if WITH_MLSDK
 #if !PLATFORM_LUMIN
@@ -777,7 +865,7 @@ void FMagicLeapController::EnumerateSources(TArray<FMotionControllerSource>& Sou
 	SourcesOut.Add(FMagicLeapMotionSourceNames::MobileApp);
 }
 
-void FMagicLeapController::UpdateControllerStateFromInputTracker(const FAppFramework& AppFramework, FName MotionSource)
+void FMagicLeapController::UpdateControllerStateFromInputTracker(const IMagicLeapPlugin& MLPlugin, FName MotionSource)
 {
 #if WITH_MLSDK
 	checkf(CurrMotionSourceControllerState.Contains(MotionSource),
@@ -827,7 +915,7 @@ void FMagicLeapController::UpdateControllerStateFromInputTracker(const FAppFrame
 			{
 				CurrControllerState.TrackingStatus = ETrackingStatus::Tracked;
 				CurrControllerState.Transform.SetLocation(MagicLeap::ToFVector(InputState.position,
-					AppFramework.GetWorldToMetersScale()));
+					MLPlugin.GetWorldToMetersScale()));
 				CurrControllerState.Transform.SetRotation(MagicLeap::ToFQuat(InputState.orientation));
 				break;
 			}
@@ -893,7 +981,7 @@ void FMagicLeapController::UpdateControllerStateFromInputTracker(const FAppFrame
 #endif //WITH_MLSDK
 }
 
-void FMagicLeapController::UpdateControllerStateFromControllerTracker(const FAppFramework& AppFramework, FName MotionSource)
+void FMagicLeapController::UpdateControllerStateFromControllerTracker(const IMagicLeapPlugin& MLPlugin, FName MotionSource)
 {
 #if WITH_MLSDK
 	// Index of the stream we're reading
@@ -937,12 +1025,12 @@ void FMagicLeapController::UpdateControllerStateFromControllerTracker(const FApp
 				ControllerState.TrackingStatus = ETrackingStatus::Tracked;
 			}
 
-			EFailReason FailReason = EFailReason::None;
-			if (!AppFramework.GetTransform(ControllerStream.
+			EMagicLeapTransformFailReason FailReason = EMagicLeapTransformFailReason::None;
+			if (!MLPlugin.GetTransform(ControllerStream.
 				coord_frame_controller, ControllerState.Transform, FailReason))
 			{
 				UE_LOG(LogMagicLeapController, Error,
-					TEXT("UpdateControllerStateFromControllerTracker: AppFramework."
+					TEXT("UpdateControllerStateFromControllerTracker: MLPlugin."
 						"GetTransform returned false, fail reason = %d."),
 					static_cast<uint32>(FailReason));
 			}
@@ -963,12 +1051,7 @@ void FMagicLeapController::UpdateTrackerData()
 		return;
 	}
 
-	const FAppFramework& AppFramework = static_cast<FMagicLeapHMD*>
-		(GEngine->XRSystem->GetHMDDevice())->GetAppFrameworkConst();
-	if (!AppFramework.IsInitialized())
-	{
-		return;
-	}
+	const IMagicLeapPlugin& MLPlugin = IMagicLeapPlugin::Get();
 
 	// First pull data from input tracker. Note that this is not conditional based on the tracking
 	// type because we also need to get buttons, touchpad, etc.
@@ -981,9 +1064,9 @@ void FMagicLeapController::UpdateTrackerData()
 			bIsInputStateValid = true;
 
 			ControllerMapper.UpdateMotionSourceInputIndexPairing(InputControllerState);
-			UpdateControllerStateFromInputTracker(AppFramework, FMagicLeapMotionSourceNames::Control0);
-			UpdateControllerStateFromInputTracker(AppFramework, FMagicLeapMotionSourceNames::Control1);
-			UpdateControllerStateFromInputTracker(AppFramework, FMagicLeapMotionSourceNames::MobileApp);
+			UpdateControllerStateFromInputTracker(MLPlugin, FMagicLeapMotionSourceNames::Control0);
+			UpdateControllerStateFromInputTracker(MLPlugin, FMagicLeapMotionSourceNames::Control1);
+			UpdateControllerStateFromInputTracker(MLPlugin, FMagicLeapMotionSourceNames::MobileApp);
 		}
 		else
 		{
@@ -996,7 +1079,7 @@ void FMagicLeapController::UpdateTrackerData()
 	}
 
 	// If mode is set to CFUID tracking overwrite the input system Dof data
-	if (TrackingMode == EMLControllerTrackingMode::CoordinateFrameUID)
+	if (TrackingMode == EMagicLeapControllerTrackingMode::CoordinateFrameUID)
 	{
 		// We need to have valid input state in order to do this, because we need to be sure
 		// what we are polling is a physical control
@@ -1006,8 +1089,8 @@ void FMagicLeapController::UpdateTrackerData()
 
 			if (MLResult_Ok == Result)
 			{
-				UpdateControllerStateFromControllerTracker(AppFramework, FMagicLeapMotionSourceNames::Control0);
-				UpdateControllerStateFromControllerTracker(AppFramework, FMagicLeapMotionSourceNames::Control1);
+				UpdateControllerStateFromControllerTracker(MLPlugin, FMagicLeapMotionSourceNames::Control0);
+				UpdateControllerStateFromControllerTracker(MLPlugin, FMagicLeapMotionSourceNames::Control1);
 			}
 			else
 			{
@@ -1020,7 +1103,7 @@ void FMagicLeapController::UpdateTrackerData()
 #endif //WITH_MLSDK
 }
 
-bool FMagicLeapController::SetControllerTrackingMode(EMLControllerTrackingMode InTrackingMode)
+bool FMagicLeapController::SetControllerTrackingMode(EMagicLeapControllerTrackingMode InTrackingMode)
 {
 #if WITH_MLSDK
 	if (IsGamepadAttached())
@@ -1036,7 +1119,7 @@ bool FMagicLeapController::SetControllerTrackingMode(EMLControllerTrackingMode I
 	return false;
 }
 
-EMLControllerTrackingMode FMagicLeapController::GetControllerTrackingMode()
+EMagicLeapControllerTrackingMode FMagicLeapController::GetControllerTrackingMode()
 {
 #if WITH_MLSDK
 	if (IsGamepadAttached())
@@ -1045,7 +1128,7 @@ EMLControllerTrackingMode FMagicLeapController::GetControllerTrackingMode()
 	}
 #endif //WITH_MLSDK
 
-	return EMLControllerTrackingMode::InputService;
+	return EMagicLeapControllerTrackingMode::InputService;
 }
 
 void FMagicLeapController::RegisterTouchpadGestureReceiver(IMagicLeapTouchpadGestures* Receiver)
@@ -1064,17 +1147,13 @@ void FMagicLeapController::UnregisterTouchpadGestureReceiver(IMagicLeapTouchpadG
 void FMagicLeapController::AddKeys()
 {
 	EKeys::AddKey(FKeyDetails(FMagicLeapKeys::MotionController_Left_Thumbstick_Z, LOCTEXT("MotionController_Left_Thumbstick_Z", "MotionController (L) Thumbstick Z"), FKeyDetails::GamepadKey | FKeyDetails::FloatAxis));
-	EKeys::AddKey(FKeyDetails(FMagicLeapKeys::Left_MoveButton, LOCTEXT("MagicLeap_Left_MoveButton", "ML (L) Move Button"), FKeyDetails::GamepadKey));
-	EKeys::AddKey(FKeyDetails(FMagicLeapKeys::Left_AppButton, LOCTEXT("MagicLeap_Left_AppButton", "ML (L) App Button"), FKeyDetails::GamepadKey));
 	EKeys::AddKey(FKeyDetails(FMagicLeapKeys::Left_HomeButton, LOCTEXT("MagicLeap_Left_HomeButton", "ML (L) Home Button"), FKeyDetails::GamepadKey));
 
 	EKeys::AddKey(FKeyDetails(FMagicLeapKeys::MotionController_Right_Thumbstick_Z, LOCTEXT("MotionController_Right_Thumbstick_Z", "MotionController (R) Thumbstick Z"), FKeyDetails::GamepadKey | FKeyDetails::FloatAxis));
-	EKeys::AddKey(FKeyDetails(FMagicLeapKeys::Right_MoveButton, LOCTEXT("MagicLeap_Right_MoveButton", "ML (R) Move Button"), FKeyDetails::GamepadKey));
-	EKeys::AddKey(FKeyDetails(FMagicLeapKeys::Right_AppButton, LOCTEXT("MagicLeap_Right_AppButton", "ML (R) App Button"), FKeyDetails::GamepadKey));
 	EKeys::AddKey(FKeyDetails(FMagicLeapKeys::Right_HomeButton, LOCTEXT("MagicLeap_Right_HomeButton", "ML (R) Home Button"), FKeyDetails::GamepadKey));
 }
 
-bool FMagicLeapController::PlayLEDPattern(FName MotionSource, EMLControllerLEDPattern LEDPattern, EMLControllerLEDColor LEDColor, float DurationInSec)
+bool FMagicLeapController::PlayLEDPattern(FName MotionSource, EMagicLeapControllerLEDPattern LEDPattern, EMagicLeapControllerLEDColor LEDColor, float DurationInSec)
 {
 #if WITH_MLSDK
 	if (IsGamepadAttached())
@@ -1111,7 +1190,7 @@ bool FMagicLeapController::PlayLEDPattern(FName MotionSource, EMLControllerLEDPa
 	return false;
 }
 
-bool FMagicLeapController::PlayLEDEffect(FName MotionSource, EMLControllerLEDEffect LEDEffect, EMLControllerLEDSpeed LEDSpeed, EMLControllerLEDPattern LEDPattern, EMLControllerLEDColor LEDColor, float DurationInSec)
+bool FMagicLeapController::PlayLEDEffect(FName MotionSource, EMagicLeapControllerLEDEffect LEDEffect, EMagicLeapControllerLEDSpeed LEDSpeed, EMagicLeapControllerLEDPattern LEDPattern, EMagicLeapControllerLEDColor LEDColor, float DurationInSec)
 {
 #if WITH_MLSDK
 	if (IsGamepadAttached())
@@ -1150,7 +1229,7 @@ bool FMagicLeapController::PlayLEDEffect(FName MotionSource, EMLControllerLEDEff
 	return false;
 }
 
-bool FMagicLeapController::PlayHapticPattern(FName MotionSource, EMLControllerHapticPattern HapticPattern, EMLControllerHapticIntensity Intensity)
+bool FMagicLeapController::PlayHapticPattern(FName MotionSource, EMagicLeapControllerHapticPattern HapticPattern, EMagicLeapControllerHapticIntensity Intensity)
 {
 #if WITH_MLSDK
 	if (IsGamepadAttached())
@@ -1186,6 +1265,16 @@ bool FMagicLeapController::PlayHapticPattern(FName MotionSource, EMLControllerHa
 	return false;
 }
 
+bool FMagicLeapController::IsMLControllerConnected(FName MotionSource) const
+{
+#if WITH_MLSDK
+	const FMagicLeapControllerState* ControllerState = CurrMotionSourceControllerState.Find(MotionSource);
+	return (ControllerState != nullptr) ? ControllerState->bIsConnected : false;
+#else
+	return false;
+#endif //WITH_MLSDK
+}
+
 void FMagicLeapController::SendControllerEventsForHand(EControllerHand Hand)
 {
 #if WITH_MLSDK
@@ -1203,7 +1292,7 @@ void FMagicLeapController::SendControllerEventsForHand(EControllerHand Hand)
 		// Touch 1 is currently not available (we have nothing to map it to)
 		if (CurrControllerState->bTouchActive[0])
 		{
-			FMagicLeapHMD::EnableInput EnableInputFromHMD;
+			MagicLeap::EnableInput EnableInputFromHMD;
 			// fixes unreferenced parameter error for Windows package builds.
 			(void)EnableInputFromHMD;
 
@@ -1216,7 +1305,7 @@ void FMagicLeapController::SendControllerEventsForHand(EControllerHand Hand)
 		}
 		else
 		{
-			FMagicLeapHMD::EnableInput EnableInputFromHMD;
+			MagicLeap::EnableInput EnableInputFromHMD;
 			// fixes unreferenced parameter error for Windows package builds.
 			(void)EnableInputFromHMD;
 
@@ -1228,7 +1317,7 @@ void FMagicLeapController::SendControllerEventsForHand(EControllerHand Hand)
 		// Analog trigger
 		if (CurrControllerState->TriggerAnalog != PrevControllerState->TriggerAnalog)
 		{
-			FMagicLeapHMD::EnableInput EnableInputFromHMD;
+			MagicLeap::EnableInput EnableInputFromHMD;
 			// fixes unreferenced parameter error for Windows package builds.
 			(void)EnableInputFromHMD;
 
@@ -1239,27 +1328,27 @@ void FMagicLeapController::SendControllerEventsForHand(EControllerHand Hand)
 #endif //WITH_MLSDK
 }
 
-EMLControllerType FMagicLeapController::GetMLControllerType(EControllerHand Hand) const
+EMagicLeapControllerType FMagicLeapController::GetMLControllerType(EControllerHand Hand) const
 {
 #if WITH_MLSDK
 	const auto& MotionSource = ControllerMapper.GetMotionSourceForHand(Hand);
 	if (MotionSource == FMagicLeapMotionSourceNames::Control0)
 	{
-		return EMLControllerType::Device;
+		return EMagicLeapControllerType::Device;
 	}
 	if (MotionSource == FMagicLeapMotionSourceNames::Control1)
 	{
-		return EMLControllerType::Device;
+		return EMagicLeapControllerType::Device;
 	}
 	if (MotionSource == FMagicLeapMotionSourceNames::MobileApp)
 	{
-		return EMLControllerType::MobileApp;
+		return EMagicLeapControllerType::MobileApp;
 	}
 #endif //WITH_MLSDK
-	return EMLControllerType::None;
+	return EMagicLeapControllerType::None;
 }
 
-bool FMagicLeapController::PlayControllerLED(EControllerHand Hand, EMLControllerLEDPattern LEDPattern, EMLControllerLEDColor LEDColor, float DurationInSec)
+bool FMagicLeapController::PlayControllerLED(EControllerHand Hand, EMagicLeapControllerLEDPattern LEDPattern, EMagicLeapControllerLEDColor LEDColor, float DurationInSec)
 {
 #if WITH_MLSDK
 	return PlayLEDPattern(ControllerMapper.GetMotionSourceForHand(Hand), LEDPattern, LEDColor, DurationInSec);
@@ -1267,7 +1356,7 @@ bool FMagicLeapController::PlayControllerLED(EControllerHand Hand, EMLController
 	return false;
 }
 
-bool FMagicLeapController::PlayControllerLEDEffect(EControllerHand Hand, EMLControllerLEDEffect LEDEffect, EMLControllerLEDSpeed LEDSpeed, EMLControllerLEDPattern LEDPattern, EMLControllerLEDColor LEDColor, float DurationInSec)
+bool FMagicLeapController::PlayControllerLEDEffect(EControllerHand Hand, EMagicLeapControllerLEDEffect LEDEffect, EMagicLeapControllerLEDSpeed LEDSpeed, EMagicLeapControllerLEDPattern LEDPattern, EMagicLeapControllerLEDColor LEDColor, float DurationInSec)
 {
 #if WITH_MLSDK
 	return PlayLEDEffect(ControllerMapper.GetMotionSourceForHand(Hand), LEDEffect, LEDSpeed, LEDPattern, LEDColor, DurationInSec);
@@ -1276,7 +1365,7 @@ bool FMagicLeapController::PlayControllerLEDEffect(EControllerHand Hand, EMLCont
 	return false;
 }
 
-bool FMagicLeapController::PlayControllerHapticFeedback(EControllerHand Hand, EMLControllerHapticPattern HapticPattern, EMLControllerHapticIntensity Intensity)
+bool FMagicLeapController::PlayControllerHapticFeedback(EControllerHand Hand, EMagicLeapControllerHapticPattern HapticPattern, EMagicLeapControllerHapticIntensity Intensity)
 {
 #if WITH_MLSDK
 	return PlayHapticPattern(ControllerMapper.GetMotionSourceForHand(Hand), HapticPattern, Intensity);
diff --git a/Engine/Plugins/Lumin/MagicLeap/Source/MagicLeapController/Private/MagicLeapController.h b/Engine/Plugins/Lumin/MagicLeap/Source/MagicLeapController/Private/MagicLeapController.h
--- a/Engine/Plugins/Lumin/MagicLeap/Source/MagicLeapController/Private/MagicLeapController.h
+++ b/Engine/Plugins/Lumin/MagicLeap/Source/MagicLeapController/Private/MagicLeapController.h
@@ -4,20 +4,14 @@
 
 #include "GenericPlatform/IInputInterface.h"
 #include "IInputDevice.h"
-#include "IMagicLeapInputDevice.h"
+#include "IMagicLeapTrackerEntity.h"
 #include "IMotionController.h"
 #include "XRMotionControllerBase.h"
 #include "IMagicLeapControllerPlugin.h"
 #include <Containers/Queue.h>
 #include "Misc/ScopeLock.h"
-#include "MagicLeapPluginUtil.h" // for ML_INCLUDES_START/END
-
-#if WITH_MLSDK
-ML_INCLUDES_START
-#include <ml_input.h>
-#include <ml_controller.h>
-ML_INCLUDES_END
-#endif //WITH_MLSDK
+#include "Lumin/CAPIShims/LuminAPIController.h"
+#include "Lumin/CAPIShims/LuminAPIInput.h"
 
 #include "CoreMinimal.h"
 #include "InputCoreTypes.h"
@@ -33,7 +27,7 @@ class IMagicLeapTouchpadGestures;
 /**
  * MagicLeap Motion Controller
  */
-class FMagicLeapController : public IMagicLeapInputDevice, public FXRMotionControllerBase
+class FMagicLeapController : public IInputDevice, public IMagicLeapTrackerEntity, public FXRMotionControllerBase
 {
 private:
 #if WITH_MLSDK
@@ -63,7 +57,7 @@ private:
 
 		EControllerHand GetHandForInputControllerIndex(uint8 controller_id) const;
 		uint8 GetInputControllerIndexForHand(EControllerHand Hand) const;
-		EMLControllerType MotionSourceToControllerType(FName InMotionSource);
+		EMagicLeapControllerType MotionSourceToControllerType(FName InMotionSource);
 
 		void SwapHands();
 	};
@@ -72,15 +66,16 @@ public:
 	FMagicLeapController(const TSharedRef<FGenericApplicationMessageHandler>& InMessageHandler);
 	virtual ~FMagicLeapController();
 
-	/** IMagicLeapInputDevice interface */
+	/** IInputDevice interface */
 	void Tick(float DeltaTime) override;
 	void SendControllerEvents() override;
 	void SetMessageHandler(const TSharedRef<FGenericApplicationMessageHandler>& InMessageHandler) override;
 	bool Exec(UWorld* InWorld, const TCHAR* Cmd, FOutputDevice& Ar) override;
 	bool IsGamepadAttached() const override;
-	void Enable() override;
-	bool SupportsExplicitEnable() const override;
-	void Disable() override;
+
+	/** IMagicLeapTrackerEntity interface */
+	void CreateEntityTracker() override;
+	void DestroyEntityTracker() override;
 
 	/** IMotionController interface */
 	bool GetControllerOrientationAndPosition(const int32 ControllerIndex, const FName MotionSource, FRotator& OutOrientation, FVector& OutPosition, float WorldToMetersScale) const override;
@@ -92,15 +87,17 @@ public:
 	void SetChannelValue(int32 ControllerId, FForceFeedbackChannelType ChannelType, float Value) override { }
 	void SetChannelValues(int32 ControllerId, const FForceFeedbackValues &values) override { }
 
-	EMLControllerTrackingMode GetControllerTrackingMode();
-	bool SetControllerTrackingMode(EMLControllerTrackingMode TrackingMode);
+	EMagicLeapControllerTrackingMode GetControllerTrackingMode();
+	bool SetControllerTrackingMode(EMagicLeapControllerTrackingMode TrackingMode);
 
 	void RegisterTouchpadGestureReceiver(IMagicLeapTouchpadGestures* Receiver);
 	void UnregisterTouchpadGestureReceiver(IMagicLeapTouchpadGestures* Receiver);
 
-	bool PlayLEDPattern(FName MotionSource, EMLControllerLEDPattern LEDPattern, EMLControllerLEDColor LEDColor, float DurationInSec);
-	bool PlayLEDEffect(FName MotionSource, EMLControllerLEDEffect LEDEffect, EMLControllerLEDSpeed LEDSpeed, EMLControllerLEDPattern LEDPattern, EMLControllerLEDColor LEDColor, float DurationInSec);
-	bool PlayHapticPattern(FName MotionSource, EMLControllerHapticPattern HapticPattern, EMLControllerHapticIntensity Intensity);
+	bool PlayLEDPattern(FName MotionSource, EMagicLeapControllerLEDPattern LEDPattern, EMagicLeapControllerLEDColor LEDColor, float DurationInSec);
+	bool PlayLEDEffect(FName MotionSource, EMagicLeapControllerLEDEffect LEDEffect, EMagicLeapControllerLEDSpeed LEDSpeed, EMagicLeapControllerLEDPattern LEDPattern, EMagicLeapControllerLEDColor LEDColor, float DurationInSec);
+	bool PlayHapticPattern(FName MotionSource, EMagicLeapControllerHapticPattern HapticPattern, EMagicLeapControllerHapticIntensity Intensity);
+
+	bool IsMLControllerConnected(FName MotionSource) const;
 
 #if WITH_MLSDK
 	// Has to be public so button functions can use it
@@ -109,15 +106,15 @@ public:
 
 	bool GetControllerOrientationAndPosition(const int32 ControllerIndex, const EControllerHand DeviceHand, FRotator& OutOrientation, FVector& OutPosition, float WorldToMetersScale) const override;
 	ETrackingStatus GetControllerTrackingStatus(const int32 ControllerIndex, const EControllerHand DeviceHand) const override;
-	EMLControllerType GetMLControllerType(EControllerHand Hand) const;
-	bool PlayControllerLED(EControllerHand Hand, EMLControllerLEDPattern LEDPattern, EMLControllerLEDColor LEDColor, float DurationInSec);
-	bool PlayControllerLEDEffect(EControllerHand Hand, EMLControllerLEDEffect LEDEffect, EMLControllerLEDSpeed LEDSpeed, EMLControllerLEDPattern LEDPattern, EMLControllerLEDColor LEDColor, float DurationInSec);
-	bool PlayControllerHapticFeedback(EControllerHand Hand, EMLControllerHapticPattern HapticPattern, EMLControllerHapticIntensity Intensity);
+	EMagicLeapControllerType GetMLControllerType(EControllerHand Hand) const;
+	bool PlayControllerLED(EControllerHand Hand, EMagicLeapControllerLEDPattern LEDPattern, EMagicLeapControllerLEDColor LEDColor, float DurationInSec);
+	bool PlayControllerLEDEffect(EControllerHand Hand, EMagicLeapControllerLEDEffect LEDEffect, EMagicLeapControllerLEDSpeed LEDSpeed, EMagicLeapControllerLEDPattern LEDPattern, EMagicLeapControllerLEDColor LEDColor, float DurationInSec);
+	bool PlayControllerHapticFeedback(EControllerHand Hand, EMagicLeapControllerHapticPattern HapticPattern, EMagicLeapControllerHapticIntensity Intensity);
 
 private:
 	void UpdateTrackerData();
-	void UpdateControllerStateFromInputTracker(const class FAppFramework& AppFramework, FName MotionSource);
-	void UpdateControllerStateFromControllerTracker(const class FAppFramework& AppFramework, FName MotionSource);
+	void UpdateControllerStateFromInputTracker(const class IMagicLeapPlugin& MLPlugin, FName MotionSource);
+	void UpdateControllerStateFromControllerTracker(const class IMagicLeapPlugin& MLPlugin, FName MotionSource);
 	void AddKeys();
 	void ReadConfigParams();
 	void InitializeInputCallbacks();
@@ -130,7 +127,7 @@ private:
 	MLHandle InputTracker;
 	MLHandle ControllerTracker;
 	MLInputControllerDof ControllerDof;
-	EMLControllerTrackingMode TrackingMode;
+	EMagicLeapControllerTrackingMode TrackingMode;
 	MLInputControllerState InputControllerState[MLInput_MaxControllers];
 	MLControllerSystemState ControllerSystemState;
 	MLInputControllerCallbacks InputControllerCallbacks;
diff --git a/Engine/Plugins/Lumin/MagicLeap/Source/MagicLeapController/Private/MagicLeapControllerFunctionLibrary.cpp b/Engine/Plugins/Lumin/MagicLeap/Source/MagicLeapController/Private/MagicLeapControllerFunctionLibrary.cpp
--- a/Engine/Plugins/Lumin/MagicLeap/Source/MagicLeapController/Private/MagicLeapControllerFunctionLibrary.cpp
+++ b/Engine/Plugins/Lumin/MagicLeap/Source/MagicLeapController/Private/MagicLeapControllerFunctionLibrary.cpp
@@ -5,7 +5,7 @@
 #include "MagicLeapController.h"
 #include "IMagicLeapPlugin.h"
 
-bool UMagicLeapControllerFunctionLibrary::PlayLEDPattern(FName MotionSource, EMLControllerLEDPattern LEDPattern, EMLControllerLEDColor LEDColor, float DurationInSec)
+bool UMagicLeapControllerFunctionLibrary::PlayLEDPattern(FName MotionSource, EMagicLeapControllerLEDPattern LEDPattern, EMagicLeapControllerLEDColor LEDColor, float DurationInSec)
 {
 	TSharedPtr<FMagicLeapController> controller = StaticCastSharedPtr<FMagicLeapController>(IMagicLeapControllerPlugin::Get().GetInputDevice());
 	if (controller.IsValid())
@@ -15,7 +15,7 @@ bool UMagicLeapControllerFunctionLibrary::PlayLEDPattern(FName MotionSource, EML
 	return false;
 }
 
-bool UMagicLeapControllerFunctionLibrary::PlayLEDEffect(FName MotionSource, EMLControllerLEDEffect LEDEffect, EMLControllerLEDSpeed LEDSpeed, EMLControllerLEDPattern LEDPattern, EMLControllerLEDColor LEDColor, float DurationInSec)
+bool UMagicLeapControllerFunctionLibrary::PlayLEDEffect(FName MotionSource, EMagicLeapControllerLEDEffect LEDEffect, EMagicLeapControllerLEDSpeed LEDSpeed, EMagicLeapControllerLEDPattern LEDPattern, EMagicLeapControllerLEDColor LEDColor, float DurationInSec)
 {
 	TSharedPtr<FMagicLeapController> controller = StaticCastSharedPtr<FMagicLeapController>(IMagicLeapControllerPlugin::Get().GetInputDevice());
 	if (controller.IsValid())
@@ -25,7 +25,7 @@ bool UMagicLeapControllerFunctionLibrary::PlayLEDEffect(FName MotionSource, EMLC
 	return false;
 }
 
-bool UMagicLeapControllerFunctionLibrary::PlayHapticPattern(FName MotionSource, EMLControllerHapticPattern HapticPattern, EMLControllerHapticIntensity Intensity)
+bool UMagicLeapControllerFunctionLibrary::PlayHapticPattern(FName MotionSource, EMagicLeapControllerHapticPattern HapticPattern, EMagicLeapControllerHapticIntensity Intensity)
 {
 	TSharedPtr<FMagicLeapController> controller = StaticCastSharedPtr<FMagicLeapController>(IMagicLeapControllerPlugin::Get().GetInputDevice());
 	if (controller.IsValid())
@@ -35,7 +35,7 @@ bool UMagicLeapControllerFunctionLibrary::PlayHapticPattern(FName MotionSource,
 	return false;
 }
 
-bool UMagicLeapControllerFunctionLibrary::SetControllerTrackingMode(EMLControllerTrackingMode TrackingMode)
+bool UMagicLeapControllerFunctionLibrary::SetControllerTrackingMode(EMagicLeapControllerTrackingMode TrackingMode)
 {
 	TSharedPtr<FMagicLeapController> controller = StaticCastSharedPtr<FMagicLeapController>(IMagicLeapControllerPlugin::Get().GetInputDevice());
 	if (controller.IsValid())
@@ -45,14 +45,14 @@ bool UMagicLeapControllerFunctionLibrary::SetControllerTrackingMode(EMLControlle
 	return false;
 }
 
-EMLControllerTrackingMode UMagicLeapControllerFunctionLibrary::GetControllerTrackingMode()
+EMagicLeapControllerTrackingMode UMagicLeapControllerFunctionLibrary::GetControllerTrackingMode()
 {
 	TSharedPtr<FMagicLeapController> controller = StaticCastSharedPtr<FMagicLeapController>(IMagicLeapControllerPlugin::Get().GetInputDevice());
 	if (controller.IsValid())
 	{
 		return controller->GetControllerTrackingMode();
 	}
-	return EMLControllerTrackingMode::InputService;
+	return EMagicLeapControllerTrackingMode::InputService;
 }
 
 FName UMagicLeapControllerFunctionLibrary::GetMotionSourceForHand(EControllerHand Hand)
@@ -92,7 +92,7 @@ bool UMagicLeapControllerFunctionLibrary::SetMotionSourceForHand(EControllerHand
 	return false;
 }
 
-EMLControllerType UMagicLeapControllerFunctionLibrary::GetControllerType(FName MotionSource)
+EMagicLeapControllerType UMagicLeapControllerFunctionLibrary::GetControllerType(FName MotionSource)
 {
 #if WITH_MLSDK
 	TSharedPtr<FMagicLeapController> controller = StaticCastSharedPtr<FMagicLeapController>(IMagicLeapControllerPlugin::Get().GetInputDevice());
@@ -101,20 +101,20 @@ EMLControllerType UMagicLeapControllerFunctionLibrary::GetControllerType(FName M
 		return controller->ControllerMapper.MotionSourceToControllerType(MotionSource);
 	}
 #endif //WITH_MLSDK
-	return EMLControllerType::None;
+	return EMagicLeapControllerType::None;
 }
 
-
-
-
-
-
-
-
-
-
-
-
+bool UMagicLeapControllerFunctionLibrary::IsMLControllerConnected(FName MotionSource)
+{
+#if WITH_MLSDK
+	TSharedPtr<FMagicLeapController> controller = StaticCastSharedPtr<FMagicLeapController>(IMagicLeapControllerPlugin::Get().GetInputDevice());
+	if (controller.IsValid())
+	{
+		return controller->IsMLControllerConnected(MotionSource);
+	}
+#endif //WITH_MLSDK
+	return false;
+}
 
 /////////////////////////////////////////////////////////////////////////////////////////////
 // DEPRECATED FUNCTIONS
@@ -144,17 +144,17 @@ void UMagicLeapControllerFunctionLibrary::InvertControllerMapping()
 #endif //WITH_MLSDK
 }
 
-EMLControllerType UMagicLeapControllerFunctionLibrary::GetMLControllerType(EControllerHand Hand)
+EMagicLeapControllerType UMagicLeapControllerFunctionLibrary::GetMLControllerType(EControllerHand Hand)
 {
 	TSharedPtr<FMagicLeapController> controller = StaticCastSharedPtr<FMagicLeapController>(IMagicLeapControllerPlugin::Get().GetInputDevice());
 	if (controller.IsValid())
 	{
 		return controller->GetMLControllerType(Hand);
 	}
-	return EMLControllerType::None;
+	return EMagicLeapControllerType::None;
 }
 
-bool UMagicLeapControllerFunctionLibrary::PlayControllerLED(EControllerHand Hand, EMLControllerLEDPattern LEDPattern, EMLControllerLEDColor LEDColor, float DurationInSec)
+bool UMagicLeapControllerFunctionLibrary::PlayControllerLED(EControllerHand Hand, EMagicLeapControllerLEDPattern LEDPattern, EMagicLeapControllerLEDColor LEDColor, float DurationInSec)
 {
 	TSharedPtr<FMagicLeapController> controller = StaticCastSharedPtr<FMagicLeapController>(IMagicLeapControllerPlugin::Get().GetInputDevice());
 	if (controller.IsValid())
@@ -164,22 +164,22 @@ bool UMagicLeapControllerFunctionLibrary::PlayControllerLED(EControllerHand Hand
 	return false;
 }
 
-bool UMagicLeapControllerFunctionLibrary::PlayControllerLEDEffect(EControllerHand Hand, EMLControllerLEDEffect LEDEffect, EMLControllerLEDSpeed LEDSpeed, EMLControllerLEDPattern LEDPattern, EMLControllerLEDColor LEDColor, float DurationInSec)
+bool UMagicLeapControllerFunctionLibrary::PlayControllerLEDEffect(EControllerHand Hand, EMagicLeapControllerLEDEffect LEDEffect, EMagicLeapControllerLEDSpeed LEDSpeed, EMagicLeapControllerLEDPattern LEDPattern, EMagicLeapControllerLEDColor LEDColor, float DurationInSec)
 {
-	TSharedPtr<FMagicLeapController> Controller = StaticCastSharedPtr<FMagicLeapController>(IMagicLeapControllerPlugin::Get().GetInputDevice());
-	if (Controller.IsValid())
+	TSharedPtr<FMagicLeapController> controller = StaticCastSharedPtr<FMagicLeapController>(IMagicLeapControllerPlugin::Get().GetInputDevice());
+	if (controller.IsValid())
 	{
-		return Controller->PlayControllerLEDEffect(Hand, LEDEffect, LEDSpeed, LEDPattern, LEDColor, DurationInSec);
+		return controller->PlayControllerLEDEffect(Hand, LEDEffect, LEDSpeed, LEDPattern, LEDColor, DurationInSec);
 	}
 	return false;
 }
 
-bool UMagicLeapControllerFunctionLibrary::PlayControllerHapticFeedback(EControllerHand Hand, EMLControllerHapticPattern HapticPattern, EMLControllerHapticIntensity Intensity)
+bool UMagicLeapControllerFunctionLibrary::PlayControllerHapticFeedback(EControllerHand Hand, EMagicLeapControllerHapticPattern HapticPattern, EMagicLeapControllerHapticIntensity Intensity)
 {
-	TSharedPtr<FMagicLeapController> Controller = StaticCastSharedPtr<FMagicLeapController>(IMagicLeapControllerPlugin::Get().GetInputDevice());
-	if (Controller.IsValid())
+	TSharedPtr<FMagicLeapController> controller = StaticCastSharedPtr<FMagicLeapController>(IMagicLeapControllerPlugin::Get().GetInputDevice());
+	if (controller.IsValid())
 	{
-		return Controller->PlayControllerHapticFeedback(Hand, HapticPattern, Intensity);
+		return controller->PlayControllerHapticFeedback(Hand, HapticPattern, Intensity);
 	}
 	return false;
 }
diff --git a/Engine/Plugins/Lumin/MagicLeap/Source/MagicLeapController/Private/MagicLeapControllerKeys.cpp b/Engine/Plugins/Lumin/MagicLeap/Source/MagicLeapController/Private/MagicLeapControllerKeys.cpp
--- a/Engine/Plugins/Lumin/MagicLeap/Source/MagicLeapController/Private/MagicLeapControllerKeys.cpp
+++ b/Engine/Plugins/Lumin/MagicLeap/Source/MagicLeapController/Private/MagicLeapControllerKeys.cpp
@@ -8,23 +8,50 @@ const FName FMagicLeapMotionSourceNames::Unknown(TEXT("Unknown"));
 const FKey FMagicLeapKeys::MotionController_Left_Thumbstick_Z("MotionController_Left_Thumbstick_Z");
 const FName FMagicLeapControllerKeyNames::MotionController_Left_Thumbstick_Z_Name("MotionController_Left_Thumbstick_Z");
 
-const FKey FMagicLeapKeys::Left_MoveButton("MagicLeap_Left_MoveButton");
-const FName FMagicLeapControllerKeyNames::Left_MoveButton_Name("MagicLeap_Left_MoveButton");
-
-const FKey FMagicLeapKeys::Left_AppButton("MagicLeap_Left_AppButton");
-const FName FMagicLeapControllerKeyNames::Left_AppButton_Name("MagicLeap_Left_AppButton");
-
 const FKey FMagicLeapKeys::Left_HomeButton("MagicLeap_Left_HomeButton");
 const FName FMagicLeapControllerKeyNames::Left_HomeButton_Name("MagicLeap_Left_HomeButton");
 
 const FKey FMagicLeapKeys::MotionController_Right_Thumbstick_Z("MotionController_Right_Thumbstick_Z");
 const FName FMagicLeapControllerKeyNames::MotionController_Right_Thumbstick_Z_Name("MotionController_Right_Thumbstick_Z");
 
-const FKey FMagicLeapKeys::Right_MoveButton("MagicLeap_Right_MoveButton");
-const FName FMagicLeapControllerKeyNames::Right_MoveButton_Name("MagicLeap_Right_MoveButton");
-
-const FKey FMagicLeapKeys::Right_AppButton("MagicLeap_Right_AppButton");
-const FName FMagicLeapControllerKeyNames::Right_AppButton_Name("MagicLeap_Right_AppButton");
-
 const FKey FMagicLeapKeys::Right_HomeButton("MagicLeap_Right_HomeButton");
 const FName FMagicLeapControllerKeyNames::Right_HomeButton_Name("MagicLeap_Right_HomeButton");
+
+const FKey FMagicLeapKeys::TouchpadGesture_Swipe_Up("TouchpadGesture_Swipe_Up");
+const FName FMagicLeapControllerKeyNames::TouchpadGesture_Swipe_Up_Name("TouchpadGesture_Swipe_Up");
+const FKey FMagicLeapKeys::TouchpadGesture_Swipe_Down("TouchpadGesture_Swipe_Down");
+const FName FMagicLeapControllerKeyNames::TouchpadGesture_Swipe_Down_Name("TouchpadGesture_Swipe_Down");
+const FKey FMagicLeapKeys::TouchpadGesture_Swipe_Right("TouchpadGesture_Swipe_Right");
+const FName FMagicLeapControllerKeyNames::TouchpadGesture_Swipe_Right_Name("TouchpadGesture_Swipe_Right");
+const FKey FMagicLeapKeys::TouchpadGesture_Swipe_Left("TouchpadGesture_Swipe_Left");
+const FName FMagicLeapControllerKeyNames::TouchpadGesture_Swipe_Left_Name("TouchpadGesture_Swipe_Left");
+const FKey FMagicLeapKeys::TouchpadGesture_RadialScroll_Clockwise("TouchpadGesture_RadialScroll_Clockwise");
+const FName FMagicLeapControllerKeyNames::TouchpadGesture_RadialScroll_Clockwise_Name("TouchpadGesture_RadialScroll_Clockwise");
+const FKey FMagicLeapKeys::TouchpadGesture_RadialScroll_CounterClockwise("TouchpadGesture_RadialScroll_CounterClockwise");
+const FName FMagicLeapControllerKeyNames::TouchpadGesture_RadialScroll_CounterClockwise_Name("TouchpadGesture_RadialScroll_CounterClockwise");
+const FKey FMagicLeapKeys::TouchpadGesture_Scroll_Up("TouchpadGesture_Scroll_Up");
+const FName FMagicLeapControllerKeyNames::TouchpadGesture_Scroll_Up_Name("TouchpadGesture_Scroll_Up");
+const FKey FMagicLeapKeys::TouchpadGesture_Scroll_Down("TouchpadGesture_Scroll_Down");
+const FName FMagicLeapControllerKeyNames::TouchpadGesture_Scroll_Down_Name("TouchpadGesture_Scroll_Down");
+const FKey FMagicLeapKeys::TouchpadGesture_Scroll_Right("TouchpadGesture_Scroll_Right");
+const FName FMagicLeapControllerKeyNames::TouchpadGesture_Scroll_Right_Name("TouchpadGesture_Scroll_Right");
+const FKey FMagicLeapKeys::TouchpadGesture_Scroll_Left("TouchpadGesture_Scroll_Left");
+const FName FMagicLeapControllerKeyNames::TouchpadGesture_Scroll_Left_Name("TouchpadGesture_Scroll_Left");
+const FKey FMagicLeapKeys::TouchpadGesture_Pinch_In("TouchpadGesture_Pinch_In");
+const FName FMagicLeapControllerKeyNames::TouchpadGesture_Pinch_In_Name("TouchpadGesture_Pinch_In");
+const FKey FMagicLeapKeys::TouchpadGesture_Pinch_Out("TouchpadGesture_Pinch_Out");
+const FName FMagicLeapControllerKeyNames::TouchpadGesture_Pinch_Out_Name("TouchpadGesture_Pinch_Out");
+const FKey FMagicLeapKeys::TouchpadGesture_LongHold("TouchpadGesture_LongHold");
+const FName FMagicLeapControllerKeyNames::TouchpadGesture_LongHold_Name("TouchpadGesture_LongHold");
+const FKey FMagicLeapKeys::TouchpadGesture_Tap("TouchpadGesture_Tap");
+const FName FMagicLeapControllerKeyNames::TouchpadGesture_Tap_Name("TouchpadGesture_Tap");
+const FKey FMagicLeapKeys::TouchpadGesture_ForceTapDown("TouchpadGesture_ForceTapDown");
+const FName FMagicLeapControllerKeyNames::TouchpadGesture_ForceTapDown_Name("TouchpadGesture_ForceTapDown");
+const FKey FMagicLeapKeys::TouchpadGesture_SecondForceDown("TouchpadGesture_SecondForceDown");
+const FName FMagicLeapControllerKeyNames::TouchpadGesture_SecondForceDown_Name("TouchpadGesture_SecondForceDown");
+const FKey FMagicLeapKeys::TouchpadGesture_ForceTapUp("TouchpadGesture_ForceTapUp");
+const FName FMagicLeapControllerKeyNames::TouchpadGesture_ForceTapUp_Name("TouchpadGesture_ForceTapUp");
+const FKey FMagicLeapKeys::TouchpadGesture_ForceDwell("TouchpadGesture_ForceDwell");
+const FName FMagicLeapControllerKeyNames::TouchpadGesture_ForceDwell_Name("TouchpadGesture_ForceDwell");
+const FKey FMagicLeapKeys::TouchpadGesture_Any("TouchpadGesture_Any");
+const FName FMagicLeapControllerKeyNames::TouchpadGesture_Any_Name("TouchpadGesture_Any");
diff --git a/Engine/Plugins/Lumin/MagicLeap/Source/MagicLeapController/Private/MagicLeapControllerMappings.h b/Engine/Plugins/Lumin/MagicLeap/Source/MagicLeapController/Private/MagicLeapControllerMappings.h
--- a/Engine/Plugins/Lumin/MagicLeap/Source/MagicLeapController/Private/MagicLeapControllerMappings.h
+++ b/Engine/Plugins/Lumin/MagicLeap/Source/MagicLeapController/Private/MagicLeapControllerMappings.h
@@ -1,49 +1,49 @@
 #if WITH_MLSDK
-#define TOUCH_GESTURE_CASE(x) case MLInputControllerTouchpadGestureType_##x: { return EMagicLeapTouchpadGestureType::x; }
+#define TOUCH_GESTURE_CASE_DEPRECATED(x) case MLInputControllerTouchpadGestureType_##x: { return EMagicLeapTouchpadGestureType::x; }
 
-EMagicLeapTouchpadGestureType MLToUnrealTouchpadGestureType(MLInputControllerTouchpadGestureType GestureType)
+EMagicLeapTouchpadGestureType MLToUnrealTouchpadGestureType_DEPRECATED(MLInputControllerTouchpadGestureType GestureType)
 {
 	switch (GestureType)
 	{
-		TOUCH_GESTURE_CASE(Tap)
-		TOUCH_GESTURE_CASE(ForceTapDown)
-		TOUCH_GESTURE_CASE(ForceTapUp)
-		TOUCH_GESTURE_CASE(ForceDwell)
-		TOUCH_GESTURE_CASE(SecondForceDown)
-		TOUCH_GESTURE_CASE(LongHold)
-		TOUCH_GESTURE_CASE(RadialScroll)
-		TOUCH_GESTURE_CASE(Swipe)
-		TOUCH_GESTURE_CASE(Scroll)
-		TOUCH_GESTURE_CASE(Pinch)
+		TOUCH_GESTURE_CASE_DEPRECATED(Tap)
+		TOUCH_GESTURE_CASE_DEPRECATED(ForceTapDown)
+		TOUCH_GESTURE_CASE_DEPRECATED(ForceTapUp)
+		TOUCH_GESTURE_CASE_DEPRECATED(ForceDwell)
+		TOUCH_GESTURE_CASE_DEPRECATED(SecondForceDown)
+		TOUCH_GESTURE_CASE_DEPRECATED(LongHold)
+		TOUCH_GESTURE_CASE_DEPRECATED(RadialScroll)
+		TOUCH_GESTURE_CASE_DEPRECATED(Swipe)
+		TOUCH_GESTURE_CASE_DEPRECATED(Scroll)
+		TOUCH_GESTURE_CASE_DEPRECATED(Pinch)
 	}
 	return EMagicLeapTouchpadGestureType::None;
 }
 
-#define TOUCH_DIR_CASE(x) case MLInputControllerTouchpadGestureDirection_##x: { return EMagicLeapTouchpadGestureDirection::x; }
+#define TOUCH_DIR_CASE_DEPRECATED(x) case MLInputControllerTouchpadGestureDirection_##x: { return EMagicLeapTouchpadGestureDirection::x; }
 
-EMagicLeapTouchpadGestureDirection MLToUnrealTouchpadGestureDirection(MLInputControllerTouchpadGestureDirection Direction)
+EMagicLeapTouchpadGestureDirection MLToUnrealTouchpadGestureDirection_DEPRECATED(MLInputControllerTouchpadGestureDirection Direction)
 {
 	switch (Direction)
 	{
-		TOUCH_DIR_CASE(Up)
-		TOUCH_DIR_CASE(Down)
-		TOUCH_DIR_CASE(Left)
-		TOUCH_DIR_CASE(Right)
-		TOUCH_DIR_CASE(In)
-		TOUCH_DIR_CASE(Out)
-		TOUCH_DIR_CASE(Clockwise)
-		TOUCH_DIR_CASE(CounterClockwise)
+		TOUCH_DIR_CASE_DEPRECATED(Up)
+		TOUCH_DIR_CASE_DEPRECATED(Down)
+		TOUCH_DIR_CASE_DEPRECATED(Left)
+		TOUCH_DIR_CASE_DEPRECATED(Right)
+		TOUCH_DIR_CASE_DEPRECATED(In)
+		TOUCH_DIR_CASE_DEPRECATED(Out)
+		TOUCH_DIR_CASE_DEPRECATED(Clockwise)
+		TOUCH_DIR_CASE_DEPRECATED(CounterClockwise)
 	}
 	return EMagicLeapTouchpadGestureDirection::None;
 }
 
-FMagicLeapTouchpadGesture MLToUnrealTouchpadGesture(EControllerHand hand, FName MotionSource, const MLInputControllerTouchpadGesture& touchpad_gesture)
+FMagicLeapTouchpadGesture MLToUnrealTouchpadGesture_DEPRECATED(EControllerHand hand, FName MotionSource, const MLInputControllerTouchpadGesture& touchpad_gesture)
 {
 	FMagicLeapTouchpadGesture gesture;
 	gesture.Hand = hand;
 	gesture.MotionSource = MotionSource;
-	gesture.Type = MLToUnrealTouchpadGestureType(touchpad_gesture.type);
-	gesture.Direction = MLToUnrealTouchpadGestureDirection(touchpad_gesture.direction);
+	gesture.Type = MLToUnrealTouchpadGestureType_DEPRECATED(touchpad_gesture.type);
+	gesture.Direction = MLToUnrealTouchpadGestureDirection_DEPRECATED(touchpad_gesture.direction);
 	gesture.PositionAndForce = FVector(touchpad_gesture.pos_and_force.x, touchpad_gesture.pos_and_force.y, touchpad_gesture.pos_and_force.z);
 	gesture.Speed = touchpad_gesture.speed;
 	gesture.Distance = touchpad_gesture.distance;
@@ -60,18 +60,6 @@ const FName& MLToUnrealButton(EControllerHand Hand, MLInputControllerButton ml_b
 
 	switch (ml_button)
 	{
-	case MLInputControllerButton_Move:
-		if (Hand == EControllerHand::Left)
-		{
-			return FMagicLeapControllerKeyNames::Left_MoveButton_Name;
-		}
-		return FMagicLeapControllerKeyNames::Right_MoveButton_Name;
-	case MLInputControllerButton_App:
-		if (Hand == EControllerHand::Left)
-		{
-			return FMagicLeapControllerKeyNames::Left_AppButton_Name;
-		}
-		return FMagicLeapControllerKeyNames::Right_AppButton_Name;
 	case MLInputControllerButton_Bumper:
 		if (Hand == EControllerHand::Left)
 		{
@@ -174,47 +162,47 @@ const FName& MLTriggerToUnrealTriggerKey(EControllerHand Hand)
 
 #if WITH_MLSDK
 
-MLInputControllerFeedbackPatternLED UnrealToMLPatternLED(EMLControllerLEDPattern LEDPattern)
+MLInputControllerFeedbackPatternLED UnrealToMLPatternLED(EMagicLeapControllerLEDPattern LEDPattern)
 {
 	switch (LEDPattern)
 	{
-	case EMLControllerLEDPattern::None:
+	case EMagicLeapControllerLEDPattern::None:
 		return MLInputControllerFeedbackPatternLED_None;
-	case EMLControllerLEDPattern::Clock01:
+	case EMagicLeapControllerLEDPattern::Clock01:
 		return MLInputControllerFeedbackPatternLED_Clock1;
-	case EMLControllerLEDPattern::Clock02:
+	case EMagicLeapControllerLEDPattern::Clock02:
 		return MLInputControllerFeedbackPatternLED_Clock2;
-	case EMLControllerLEDPattern::Clock03:
+	case EMagicLeapControllerLEDPattern::Clock03:
 		return MLInputControllerFeedbackPatternLED_Clock3;
-	case EMLControllerLEDPattern::Clock04:
+	case EMagicLeapControllerLEDPattern::Clock04:
 		return MLInputControllerFeedbackPatternLED_Clock4;
-	case EMLControllerLEDPattern::Clock05:
+	case EMagicLeapControllerLEDPattern::Clock05:
 		return MLInputControllerFeedbackPatternLED_Clock5;
-	case EMLControllerLEDPattern::Clock06:
+	case EMagicLeapControllerLEDPattern::Clock06:
 		return MLInputControllerFeedbackPatternLED_Clock6;
-	case EMLControllerLEDPattern::Clock07:
+	case EMagicLeapControllerLEDPattern::Clock07:
 		return MLInputControllerFeedbackPatternLED_Clock7;
-	case EMLControllerLEDPattern::Clock08:
+	case EMagicLeapControllerLEDPattern::Clock08:
 		return MLInputControllerFeedbackPatternLED_Clock8;
-	case EMLControllerLEDPattern::Clock09:
+	case EMagicLeapControllerLEDPattern::Clock09:
 		return MLInputControllerFeedbackPatternLED_Clock9;
-	case EMLControllerLEDPattern::Clock10:
+	case EMagicLeapControllerLEDPattern::Clock10:
 		return MLInputControllerFeedbackPatternLED_Clock10;
-	case EMLControllerLEDPattern::Clock11:
+	case EMagicLeapControllerLEDPattern::Clock11:
 		return MLInputControllerFeedbackPatternLED_Clock11;
-	case EMLControllerLEDPattern::Clock12:
+	case EMagicLeapControllerLEDPattern::Clock12:
 		return MLInputControllerFeedbackPatternLED_Clock12;
-	case EMLControllerLEDPattern::Clock01_07:
+	case EMagicLeapControllerLEDPattern::Clock01_07:
 		return MLInputControllerFeedbackPatternLED_Clock1And7;
-	case EMLControllerLEDPattern::Clock02_08:
+	case EMagicLeapControllerLEDPattern::Clock02_08:
 		return MLInputControllerFeedbackPatternLED_Clock2And8;
-	case EMLControllerLEDPattern::Clock03_09:
+	case EMagicLeapControllerLEDPattern::Clock03_09:
 		return MLInputControllerFeedbackPatternLED_Clock3And9;
-	case EMLControllerLEDPattern::Clock04_10:
+	case EMagicLeapControllerLEDPattern::Clock04_10:
 		return MLInputControllerFeedbackPatternLED_Clock4And10;
-	case EMLControllerLEDPattern::Clock05_11:
+	case EMagicLeapControllerLEDPattern::Clock05_11:
 		return MLInputControllerFeedbackPatternLED_Clock5And11;
-	case EMLControllerLEDPattern::Clock06_12:
+	case EMagicLeapControllerLEDPattern::Clock06_12:
 		return MLInputControllerFeedbackPatternLED_Clock6And12;
 	default:
 		UE_LOG(LogMagicLeapController, Error, TEXT("Unhandled LED Pattern type %d"), static_cast<int32>(LEDPattern));
@@ -223,21 +211,21 @@ MLInputControllerFeedbackPatternLED UnrealToMLPatternLED(EMLControllerLEDPattern
 	return MLInputControllerFeedbackPatternLED_Ensure32Bits;
 }
 
-MLInputControllerFeedbackEffectLED UnrealToMLEffectLED(EMLControllerLEDEffect LEDEffect)
+MLInputControllerFeedbackEffectLED UnrealToMLEffectLED(EMagicLeapControllerLEDEffect LEDEffect)
 {
 	switch (LEDEffect)
 	{
-	case EMLControllerLEDEffect::RotateCW:
+	case EMagicLeapControllerLEDEffect::RotateCW:
 		return MLInputControllerFeedbackEffectLED_RotateCW;
-	case EMLControllerLEDEffect::RotateCCW:
+	case EMagicLeapControllerLEDEffect::RotateCCW:
 		return MLInputControllerFeedbackEffectLED_RotateCCW;
-	case EMLControllerLEDEffect::Pulse:
+	case EMagicLeapControllerLEDEffect::Pulse:
 		return MLInputControllerFeedbackEffectLED_Pulse;
-	case EMLControllerLEDEffect::PaintCW:
+	case EMagicLeapControllerLEDEffect::PaintCW:
 		return MLInputControllerFeedbackEffectLED_PaintCW;
-	case EMLControllerLEDEffect::PaintCCW:
+	case EMagicLeapControllerLEDEffect::PaintCCW:
 		return MLInputControllerFeedbackEffectLED_PaintCCW;
-	case EMLControllerLEDEffect::Blink:
+	case EMagicLeapControllerLEDEffect::Blink:
 		return MLInputControllerFeedbackEffectLED_Blink;
 	default:
 		UE_LOG(LogMagicLeapController, Error, TEXT("Unhandled LED effect type %d"), static_cast<int32>(LEDEffect));
@@ -246,8 +234,8 @@ MLInputControllerFeedbackEffectLED UnrealToMLEffectLED(EMLControllerLEDEffect LE
 	return MLInputControllerFeedbackEffectLED_Ensure32Bits;
 }
 
-#define LED_COLOR_CASE(x) case EMLControllerLEDColor::x: { return MLInputControllerFeedbackColorLED_##x; }
-MLInputControllerFeedbackColorLED UnrealToMLColorLED(EMLControllerLEDColor LEDColor)
+#define LED_COLOR_CASE(x) case EMagicLeapControllerLEDColor::x: { return MLInputControllerFeedbackColorLED_##x; }
+MLInputControllerFeedbackColorLED UnrealToMLColorLED(EMagicLeapControllerLEDColor LEDColor)
 {
 	switch (LEDColor)
 	{
@@ -274,15 +262,15 @@ MLInputControllerFeedbackColorLED UnrealToMLColorLED(EMLControllerLEDColor LEDCo
 	return MLInputControllerFeedbackColorLED_Ensure32Bits;
 }
 
-MLInputControllerFeedbackEffectSpeedLED UnrealToMLSpeedLED(EMLControllerLEDSpeed LEDSpeed)
+MLInputControllerFeedbackEffectSpeedLED UnrealToMLSpeedLED(EMagicLeapControllerLEDSpeed LEDSpeed)
 {
 	switch (LEDSpeed)
 	{
-	case EMLControllerLEDSpeed::Slow:
+	case EMagicLeapControllerLEDSpeed::Slow:
 		return MLInputControllerFeedbackEffectSpeedLED_Slow;
-	case EMLControllerLEDSpeed::Medium:
+	case EMagicLeapControllerLEDSpeed::Medium:
 		return MLInputControllerFeedbackEffectSpeedLED_Medium;
-	case EMLControllerLEDSpeed::Fast:
+	case EMagicLeapControllerLEDSpeed::Fast:
 		return MLInputControllerFeedbackEffectSpeedLED_Fast;
 	default:
 		UE_LOG(LogMagicLeapController, Error, TEXT("Unhandled LED speed type %d"), static_cast<int32>(LEDSpeed));
@@ -291,29 +279,29 @@ MLInputControllerFeedbackEffectSpeedLED UnrealToMLSpeedLED(EMLControllerLEDSpeed
 	return MLInputControllerFeedbackEffectSpeedLED_Ensure32Bits;
 }
 
-MLInputControllerFeedbackPatternVibe UnrealToMLPatternVibe(EMLControllerHapticPattern HapticPattern)
+MLInputControllerFeedbackPatternVibe UnrealToMLPatternVibe(EMagicLeapControllerHapticPattern HapticPattern)
 {
 	switch (HapticPattern)
 	{
-	case EMLControllerHapticPattern::None:
+	case EMagicLeapControllerHapticPattern::None:
 		return MLInputControllerFeedbackPatternVibe_None;
-	case EMLControllerHapticPattern::Click:
+	case EMagicLeapControllerHapticPattern::Click:
 		return MLInputControllerFeedbackPatternVibe_Click;
-	case EMLControllerHapticPattern::Bump:
+	case EMagicLeapControllerHapticPattern::Bump:
 		return MLInputControllerFeedbackPatternVibe_Bump;
-	case EMLControllerHapticPattern::DoubleClick:
+	case EMagicLeapControllerHapticPattern::DoubleClick:
 		return MLInputControllerFeedbackPatternVibe_DoubleClick;
-	case EMLControllerHapticPattern::Buzz:
+	case EMagicLeapControllerHapticPattern::Buzz:
 		return MLInputControllerFeedbackPatternVibe_Buzz;
-	case EMLControllerHapticPattern::Tick:
+	case EMagicLeapControllerHapticPattern::Tick:
 		return MLInputControllerFeedbackPatternVibe_Tick;
-	case EMLControllerHapticPattern::ForceDown:
+	case EMagicLeapControllerHapticPattern::ForceDown:
 		return MLInputControllerFeedbackPatternVibe_ForceDown;
-	case EMLControllerHapticPattern::ForceUp:
+	case EMagicLeapControllerHapticPattern::ForceUp:
 		return MLInputControllerFeedbackPatternVibe_ForceUp;
-	case EMLControllerHapticPattern::ForceDwell:
+	case EMagicLeapControllerHapticPattern::ForceDwell:
 		return MLInputControllerFeedbackPatternVibe_ForceDwell;
-	case EMLControllerHapticPattern::SecondForceDown:
+	case EMagicLeapControllerHapticPattern::SecondForceDown:
 		return MLInputControllerFeedbackPatternVibe_SecondForceDown;
 	default:
 		UE_LOG(LogMagicLeapController, Error, TEXT("Unhandled Haptic Pattern type %d"), static_cast<int32>(HapticPattern));
@@ -322,15 +310,15 @@ MLInputControllerFeedbackPatternVibe UnrealToMLPatternVibe(EMLControllerHapticPa
 	return MLInputControllerFeedbackPatternVibe_Ensure32Bits;
 }
 
-MLInputControllerFeedbackIntensity UnrealToMLHapticIntensity(EMLControllerHapticIntensity HapticIntensity)
+MLInputControllerFeedbackIntensity UnrealToMLHapticIntensity(EMagicLeapControllerHapticIntensity HapticIntensity)
 {
 	switch (HapticIntensity)
 	{
-	case EMLControllerHapticIntensity::Low:
+	case EMagicLeapControllerHapticIntensity::Low:
 		return MLInputControllerFeedbackIntensity_Low;
-	case EMLControllerHapticIntensity::Medium:
+	case EMagicLeapControllerHapticIntensity::Medium:
 		return MLInputControllerFeedbackIntensity_Medium;
-	case EMLControllerHapticIntensity::High:
+	case EMagicLeapControllerHapticIntensity::High:
 		return MLInputControllerFeedbackIntensity_High;
 	default:
 		UE_LOG(LogMagicLeapController, Error, TEXT("Unhandled Haptic Intensity type %d"), static_cast<int32>(HapticIntensity));
diff --git a/Engine/Plugins/Lumin/MagicLeap/Source/MagicLeapController/Private/MagicLeapControllerPlugin.cpp b/Engine/Plugins/Lumin/MagicLeap/Source/MagicLeapController/Private/MagicLeapControllerPlugin.cpp
--- a/Engine/Plugins/Lumin/MagicLeap/Source/MagicLeapController/Private/MagicLeapControllerPlugin.cpp
+++ b/Engine/Plugins/Lumin/MagicLeap/Source/MagicLeapController/Private/MagicLeapControllerPlugin.cpp
@@ -1,6 +1,7 @@
 #include "IMagicLeapControllerPlugin.h"
 #include "MagicLeapController.h"
 #include "Framework/Application/SlateApplication.h"
+#include "MagicLeapPluginUtil.h"
 
 class FMagicLeapControllerPlugin : public IMagicLeapControllerPlugin
 {
@@ -26,17 +27,6 @@ public:
 		// Unfortunately, the UObject system hasn't finished initialization when this module has been loaded.
 		bool bIsVDZIEnabled = false;
 		GConfig->GetBool(TEXT("/Script/MagicLeap.MagicLeapSettings"), TEXT("bEnableZI"), bIsVDZIEnabled, GEngineIni);
-
-		APISetup.Startup(bIsVDZIEnabled);
-#if WITH_MLSDK
-		APISetup.LoadDLL(TEXT("ml_input"));
-#endif
-	}
-
-	virtual void ShutdownModule() override
-	{
-		APISetup.Shutdown();
-		IMagicLeapControllerPlugin::ShutdownModule();
 	}
 
 	virtual TSharedPtr<class IInputDevice> CreateInputDevice(const TSharedRef<FGenericApplicationMessageHandler>& InMessageHandler) override
