Copying //UE4/Partner-MagicLeap-Staging to Dev-VR (//UE4/Dev-VR)

diff --git a/Engine/Plugins/Lumin/MagicLeap/Source/MagicLeapController/Private/Tests/MagicLeapControllerTests.cpp b/Engine/Plugins/Lumin/MagicLeap/Source/MagicLeapController/Private/Tests/MagicLeapControllerTests.cpp
--- a/Engine/Plugins/Lumin/MagicLeap/Source/MagicLeapController/Private/Tests/MagicLeapControllerTests.cpp
+++ b/Engine/Plugins/Lumin/MagicLeap/Source/MagicLeapController/Private/Tests/MagicLeapControllerTests.cpp
@@ -4,20 +4,14 @@
 #include "Misc/AutomationTest.h"
 #include "Tests/AutomationCommon.h"
 #include "GenericPlatform/IInputInterface.h"
-#include "IMagicLeapInputDevice.h"
+#include "IMagicLeapTrackerEntity.h"
 #include "IHapticDevice.h"
 #include "IMotionController.h"
 #include "MagicLeapControllerKeys.h"
 #include "MagicLeapControllerFunctionLibrary.h"
 #include "Framework/Application/SlateApplication.h"
 #include "Haptics/HapticFeedbackEffect_Base.h"
-#include "MagicLeapPluginUtil.h" // for ML_INCLUDES_START/END
-
-#if WITH_MLSDK
-ML_INCLUDES_START
-#include <ml_input.h>
-ML_INCLUDES_END
-#endif //WITH_MLSDK
+#include "Lumin/CAPIShims/LuminAPI.h"
 
 DEFINE_LOG_CATEGORY_STATIC(LogMagicLeapControllerTest, Display, All);
 
@@ -32,45 +26,53 @@ void MagicLeapTestReferenceFunction()
 /**
 * Play Pattern Haptic Effect
 */
-DEFINE_LATENT_AUTOMATION_COMMAND_THREE_PARAMETER(FPlayPatternHapticLatentCommand, EControllerHand, Hand, EMLControllerHapticPattern, Pattern, EMLControllerHapticIntensity, Intensity);
+DEFINE_LATENT_AUTOMATION_COMMAND_THREE_PARAMETER(FPlayPatternHapticLatentCommand, FName, Hand, EMagicLeapControllerHapticPattern, Pattern, EMagicLeapControllerHapticIntensity, Intensity);
 bool FPlayPatternHapticLatentCommand::Update()
 {
 	FString PatternName;
-	UEnum* PatternEnum = FindObjectChecked<UEnum>(ANY_PACKAGE, TEXT("EMLControllerHapticPattern"), true);
+	UEnum* PatternEnum = FindObjectChecked<UEnum>(ANY_PACKAGE, TEXT("EMagicLeapControllerHapticPattern"), true);
 	if (PatternEnum != nullptr)
 	{
 		PatternName = PatternEnum->GetNameByValue((int32)Pattern).ToString();
 	}
 
-	UE_LOG(LogCore, Log, TEXT("FPlayPatternHapticLatentCommand %d, %s"), (int32)Hand, *PatternName);
+	FString IntensityName;
+	UEnum* IntensityEnum = FindObjectChecked<UEnum>(ANY_PACKAGE, TEXT("EMagicLeapControllerHapticIntensity"), true);
+	if (IntensityEnum != nullptr)
+	{
+		IntensityName = IntensityEnum->GetNameByValue((int32)Intensity).ToString();
+	}
+
+	UE_LOG(LogCore, Log, TEXT("FPlayPatternHapticLatentCommand %s, %s %s"), *Hand.ToString(), *PatternName, *IntensityName);
 
-	UMagicLeapControllerFunctionLibrary::PlayControllerHapticFeedback(Hand, Pattern, Intensity);
+	UMagicLeapControllerFunctionLibrary::PlayHapticPattern(Hand, Pattern, Intensity);
 
 	return true;
 }
 
 
 //@TODO - remove editor
-IMPLEMENT_SIMPLE_AUTOMATION_TEST(FMagicLeapControllerHapticTest, "System.VR.MagicLeap.Haptics.Pattern", EAutomationTestFlags::ClientContext | EAutomationTestFlags::EditorContext | EAutomationTestFlags::EngineFilter)
+IMPLEMENT_SIMPLE_AUTOMATION_TEST(FMagicLeapControllerHapticTest, "System.VR.MagicLeap.Haptics.Patterns", EAutomationTestFlags::ClientContext | EAutomationTestFlags::EditorContext | EAutomationTestFlags::EngineFilter)
 bool FMagicLeapControllerHapticTest::RunTest(const FString&)
 {
-	UEnum* PatternEnum = FindObjectChecked<UEnum>(ANY_PACKAGE, TEXT("EMLControllerHapticPattern"), true);
-	UEnum* IntensityEnum = FindObjectChecked<UEnum>(ANY_PACKAGE, TEXT("EMLControllerHapticIntensity"), true);
+	UEnum* PatternEnum = FindObjectChecked<UEnum>(ANY_PACKAGE, TEXT("EMagicLeapControllerHapticPattern"), true);
+	UEnum* IntensityEnum = FindObjectChecked<UEnum>(ANY_PACKAGE, TEXT("EMagicLeapControllerHapticIntensity"), true);
 
-	float ActiveDuration = 1.0f;
+	float ActiveDuration = 2.0f;
 	//for (uint32 HandIndex = 0; HandIndex < MLInput_MaxControllers; ++HandIndex)
 	//for now only right hand matters
-	uint32 HandIndex = (uint32)EControllerHand::Right;
+	FName MotionSource = TEXT("Right");
 	{
-		for (uint32 PatternIndex = 0; PatternIndex < PatternEnum->GetMaxEnumValue(); ++PatternIndex)
+		// haptic pattern 0 is None, skip it
+		for (uint32 PatternIndex = 1; PatternIndex < PatternEnum->GetMaxEnumValue(); ++PatternIndex)
 		{
 			for (uint32 IntensityIndex = 0; IntensityIndex < IntensityEnum->GetMaxEnumValue(); ++IntensityIndex)
 			{
 				//Turn on haptics
 				ADD_LATENT_AUTOMATION_COMMAND(FPlayPatternHapticLatentCommand(
-					(EControllerHand)HandIndex,
-					(EMLControllerHapticPattern)PatternIndex,
-					(EMLControllerHapticIntensity)IntensityIndex));
+					MotionSource,
+					(EMagicLeapControllerHapticPattern)PatternIndex,
+					(EMagicLeapControllerHapticIntensity)IntensityIndex));
 				//Give the command a chance to finish
 				ADD_LATENT_AUTOMATION_COMMAND(FWaitLatentCommand(ActiveDuration));
 			}
@@ -82,53 +84,54 @@ bool FMagicLeapControllerHapticTest::RunTest(const FString&)
 
 
 /**
-* Play LED Effect
+* Play LED Pattern
 */
-DEFINE_LATENT_AUTOMATION_COMMAND_FOUR_PARAMETER(FPlayLEDPatternLatentCommand, EControllerHand, Hand, EMLControllerLEDPattern, Pattern, EMLControllerLEDColor, Color, float, Duration);
+DEFINE_LATENT_AUTOMATION_COMMAND_FOUR_PARAMETER(FPlayLEDPatternLatentCommand, FName, Hand, EMagicLeapControllerLEDPattern, Pattern, EMagicLeapControllerLEDColor, Color, float, Duration);
 bool FPlayLEDPatternLatentCommand::Update()
 {
 	FString PatternName;
-	UEnum* PatternEnum = FindObjectChecked<UEnum>(ANY_PACKAGE, TEXT("EMLControllerLEDPattern"), true);
+	UEnum* PatternEnum = FindObjectChecked<UEnum>(ANY_PACKAGE, TEXT("EMagicLeapControllerLEDPattern"), true);
 	if (PatternEnum != nullptr)
 	{
 		PatternName = PatternEnum->GetNameByValue((int32)Pattern).ToString();
 	}
 
 	FString ColorName;
-	UEnum* ColorEnum = FindObjectChecked<UEnum>(ANY_PACKAGE, TEXT("EMLControllerLEDColor"), true);
+	UEnum* ColorEnum = FindObjectChecked<UEnum>(ANY_PACKAGE, TEXT("EMagicLeapControllerLEDColor"), true);
 	if (ColorEnum != nullptr)
 	{
 		ColorName = ColorEnum->GetNameByValue((int32)Color).ToString();
 	}
 
-	UE_LOG(LogCore, Log, TEXT("FPlayLEDPatternLatentCommand %d Hand, %s %s"), (int32)Hand, *PatternName, *ColorName);
+	UE_LOG(LogCore, Log, TEXT("FPlayLEDPatternLatentCommand %s Hand, %s %s"), *Hand.ToString(), *PatternName, *ColorName);
 
-	UMagicLeapControllerFunctionLibrary::PlayControllerLED(Hand, Pattern, Color, Duration);
+	UMagicLeapControllerFunctionLibrary::PlayLEDPattern(Hand, Pattern, Color, Duration);
 
 	return true;
 }
 
 
 //@TODO - remove editor
-IMPLEMENT_SIMPLE_AUTOMATION_TEST(FMagicLeapControllerLEDTest, "System.VR.MagicLeap.LED", EAutomationTestFlags::ClientContext | EAutomationTestFlags::EditorContext | EAutomationTestFlags::EngineFilter)
-bool FMagicLeapControllerLEDTest::RunTest(const FString&)
+IMPLEMENT_SIMPLE_AUTOMATION_TEST(FMagicLeapControllerLEDPatternTest, "System.VR.MagicLeap.LED.Patterns", EAutomationTestFlags::ClientContext | EAutomationTestFlags::EditorContext | EAutomationTestFlags::EngineFilter)
+bool FMagicLeapControllerLEDPatternTest::RunTest(const FString&)
 {
-	UEnum* PatternEnum = FindObjectChecked<UEnum>(ANY_PACKAGE, TEXT("EMLControllerLEDPattern"), true);
-	UEnum* ColorEnum = FindObjectChecked<UEnum>(ANY_PACKAGE, TEXT("EMLControllerLEDColor"), true);
+	UEnum* PatternEnum = FindObjectChecked<UEnum>(ANY_PACKAGE, TEXT("EMagicLeapControllerLEDPattern"), true);
+	UEnum* ColorEnum = FindObjectChecked<UEnum>(ANY_PACKAGE, TEXT("EMagicLeapControllerLEDColor"), true);
 
 	float ActiveDuration = 1.0f;
 	float InactiveDuration = 0.5f;
 	//run through each pattern and run the pattern, wait for it to finish, turn off, wait
 	//for (uint32 HandIndex = 0; HandIndex < MLInput_MaxControllers; ++HandIndex)
 	//for now only right hand matters
-	uint32 HandIndex = (uint32)EControllerHand::Right;
+	FName MotionSource = TEXT("Right");
 	{
-		for (uint32 PatternIndex = 0; PatternIndex < PatternEnum->GetMaxEnumValue(); ++PatternIndex)
+		// led pattern 0 is None, skip it
+		for (uint32 PatternIndex = 1; PatternIndex < PatternEnum->GetMaxEnumValue(); ++PatternIndex)
 		{
 			for (uint32 ColorIndex = 0; ColorIndex < ColorEnum->GetMaxEnumValue(); ++ColorIndex)
 			{
 				//Turn LED pattern on
-				ADD_LATENT_AUTOMATION_COMMAND(FPlayLEDPatternLatentCommand((EControllerHand)HandIndex, (EMLControllerLEDPattern)PatternIndex, (EMLControllerLEDColor)ColorIndex, ActiveDuration));
+				ADD_LATENT_AUTOMATION_COMMAND(FPlayLEDPatternLatentCommand(MotionSource, (EMagicLeapControllerLEDPattern)PatternIndex, (EMagicLeapControllerLEDColor)ColorIndex, ActiveDuration));
 				//Give the command a chance to finish
 				ADD_LATENT_AUTOMATION_COMMAND(FWaitLatentCommand(ActiveDuration));
 				//Add a second to delimit between patterns
@@ -140,4 +143,87 @@ bool FMagicLeapControllerLEDTest::RunTest(const FString&)
 	return true;
 }
 
+
+/**
+* Play LED Effect
+*/
+DEFINE_LATENT_AUTOMATION_COMMAND_FIVE_PARAMETER(FPlayLEDEffectLatentCommand, FName, Hand, EMagicLeapControllerLEDEffect, LEDEffect, EMagicLeapControllerLEDSpeed, LEDSpeed, EMagicLeapControllerLEDPattern, Pattern, EMagicLeapControllerLEDColor, Color);
+bool FPlayLEDEffectLatentCommand::Update()
+{
+	FString EffectName;
+	UEnum* EffectEnum = FindObjectChecked<UEnum>(ANY_PACKAGE, TEXT("EMagicLeapControllerLEDEffect"), true);
+	if (EffectEnum != nullptr)
+	{
+		EffectName = EffectEnum->GetNameByValue((int32)LEDEffect).ToString();
+	}
+
+	FString SpeedName;
+	UEnum* SpeedEnum = FindObjectChecked<UEnum>(ANY_PACKAGE, TEXT("EMagicLeapControllerLEDSpeed"), true);
+	if (SpeedEnum != nullptr)
+	{
+		SpeedName = SpeedEnum->GetNameByValue((int32)LEDSpeed).ToString();
+	}
+
+	FString PatternName;
+	UEnum* PatternEnum = FindObjectChecked<UEnum>(ANY_PACKAGE, TEXT("EMagicLeapControllerLEDPattern"), true);
+	if (PatternEnum != nullptr)
+	{
+		PatternName = PatternEnum->GetNameByValue((int32)Pattern).ToString();
+	}
+
+	FString ColorName;
+	UEnum* ColorEnum = FindObjectChecked<UEnum>(ANY_PACKAGE, TEXT("EMagicLeapControllerLEDColor"), true);
+	if (ColorEnum != nullptr)
+	{
+		ColorName = ColorEnum->GetNameByValue((int32)Color).ToString();
+	}
+
+	UE_LOG(LogCore, Log, TEXT("FPlayLEDEffectLatentCommand %s Hand, $s %s %s %s"), *Hand.ToString(), *EffectName, *SpeedName, *PatternName, *ColorName);
+
+	UMagicLeapControllerFunctionLibrary::PlayLEDEffect(Hand, LEDEffect, LEDSpeed, Pattern, Color, 1.0f);
+
+	return true;
+}
+
+
+//@TODO - remove editor
+IMPLEMENT_SIMPLE_AUTOMATION_TEST(FMagicLeapControllerLEDEffectTest, "System.VR.MagicLeap.LED.Effects", EAutomationTestFlags::ClientContext | EAutomationTestFlags::EditorContext | EAutomationTestFlags::EngineFilter)
+bool FMagicLeapControllerLEDEffectTest::RunTest(const FString&)
+{
+	UEnum* EffectEnum = FindObjectChecked<UEnum>(ANY_PACKAGE, TEXT("EMagicLeapControllerLEDEffect"), true);
+	UEnum* SpeedEnum = FindObjectChecked<UEnum>(ANY_PACKAGE, TEXT("EMagicLeapControllerLEDSpeed"), true);
+	UEnum* PatternEnum = FindObjectChecked<UEnum>(ANY_PACKAGE, TEXT("EMagicLeapControllerLEDPattern"), true);
+	UEnum* ColorEnum = FindObjectChecked<UEnum>(ANY_PACKAGE, TEXT("EMagicLeapControllerLEDColor"), true);
+
+	float ActiveDuration = 1.0f;
+	float InactiveDuration = 0.5f;
+	//run through each effect, speed, pattern, color and run the effect, wait for it to finish, turn off, wait
+	//for (uint32 HandIndex = 0; HandIndex < MLInput_MaxControllers; ++HandIndex)
+	//for now only right hand matters
+	FName MotionSource = TEXT("Right");
+	{
+		for (uint32 EffectIndex = 0; EffectIndex < EffectEnum->GetMaxEnumValue(); ++EffectIndex)
+		{
+			for (uint32 SpeedIndex = 0; SpeedIndex < SpeedEnum->GetMaxEnumValue(); ++SpeedIndex)
+			{
+				// led pattern 0 is None, skip it
+				for (uint32 PatternIndex = 1; PatternIndex < PatternEnum->GetMaxEnumValue(); ++PatternIndex)
+				{
+					for (uint32 ColorIndex = 0; ColorIndex < ColorEnum->GetMaxEnumValue(); ++ColorIndex)
+					{
+						//Turn LED pattern on
+						ADD_LATENT_AUTOMATION_COMMAND(FPlayLEDEffectLatentCommand(MotionSource, (EMagicLeapControllerLEDEffect)EffectIndex, (EMagicLeapControllerLEDSpeed)SpeedIndex, (EMagicLeapControllerLEDPattern)PatternIndex, (EMagicLeapControllerLEDColor)ColorIndex));
+						//Give the command a chance to finish
+						ADD_LATENT_AUTOMATION_COMMAND(FWaitLatentCommand(ActiveDuration));
+						//Add a second to delimit between patterns
+						ADD_LATENT_AUTOMATION_COMMAND(FWaitLatentCommand(InactiveDuration));
+					}
+				}
+			}
+		}
+	}
+
+	return true;
+}
+
 #endif
\ No newline at end of file
diff --git a/Engine/Plugins/Lumin/MagicLeap/Source/MagicLeapController/Private/TouchpadGesturesComponent.cpp b/Engine/Plugins/Lumin/MagicLeap/Source/MagicLeapController/Private/TouchpadGesturesComponent.cpp
--- a/Engine/Plugins/Lumin/MagicLeap/Source/MagicLeapController/Private/TouchpadGesturesComponent.cpp
+++ b/Engine/Plugins/Lumin/MagicLeap/Source/MagicLeapController/Private/TouchpadGesturesComponent.cpp
@@ -4,55 +4,45 @@
 #include "IMagicLeapControllerPlugin.h"
 #include "MagicLeapController.h"
 
-UTouchpadGesturesComponent::UTouchpadGesturesComponent()
+UMagicLeapTouchpadGesturesComponent::UMagicLeapTouchpadGesturesComponent()
 {
 	// Make sure this component ticks
 	PrimaryComponentTick.bCanEverTick = true;
 	PrimaryComponentTick.bStartWithTickEnabled = true;
 	PrimaryComponentTick.TickGroup = TG_PrePhysics;
 	bAutoActivate = true;
-
-	if (!HasAnyFlags(RF_ClassDefaultObject))
-	{
-		TSharedPtr<FMagicLeapController> controller = StaticCastSharedPtr<FMagicLeapController>(IMagicLeapControllerPlugin::Get().GetInputDevice());
-		if (controller.IsValid())
-		{
-			controller->RegisterTouchpadGestureReceiver(this);
-		}
-	}
-}
-
-UTouchpadGesturesComponent::~UTouchpadGesturesComponent()
-{
-	if (!HasAnyFlags(RF_ClassDefaultObject))
-	{
-		TSharedPtr<FMagicLeapController> controller = StaticCastSharedPtr<FMagicLeapController>(IMagicLeapControllerPlugin::Get().GetInputDevice());
-		if (controller.IsValid())
-		{
-			controller->UnregisterTouchpadGestureReceiver(this);
-		}
-	}
 }
 
-void UTouchpadGesturesComponent::OnTouchpadGestureStartCallback(const FMagicLeapTouchpadGesture& GestureData)
+void UMagicLeapTouchpadGesturesComponent::OnTouchpadGestureStartCallback(const FMagicLeapTouchpadGesture& GestureData)
 {
 	FScopeLock Lock(&CriticalSection);
 	PendingTouchpadGestureStart.Add(GestureData);
 }
 
-void UTouchpadGesturesComponent::OnTouchpadGestureContinueCallback(const FMagicLeapTouchpadGesture& GestureData)
+void UMagicLeapTouchpadGesturesComponent::OnTouchpadGestureContinueCallback(const FMagicLeapTouchpadGesture& GestureData)
 {
 	FScopeLock Lock(&CriticalSection);
 	PendingTouchpadGestureContinue.Add(GestureData);
 }
 
-void UTouchpadGesturesComponent::OnTouchpadGestureEndCallback(const FMagicLeapTouchpadGesture& GestureData)
+void UMagicLeapTouchpadGesturesComponent::OnTouchpadGestureEndCallback(const FMagicLeapTouchpadGesture& GestureData)
 {
 	FScopeLock Lock(&CriticalSection);
 	PendingTouchpadGestureEnd.Add(GestureData);
 }
 
-void UTouchpadGesturesComponent::TickComponent(float DeltaTime, enum ELevelTick TickType, FActorComponentTickFunction *ThisTickFunction)
+void UMagicLeapTouchpadGesturesComponent::BeginPlay()
+{
+	Super::BeginPlay();
+
+	TSharedPtr<FMagicLeapController> controller = StaticCastSharedPtr<FMagicLeapController>(IMagicLeapControllerPlugin::Get().GetInputDevice());
+	if (controller.IsValid())
+	{
+		controller->RegisterTouchpadGestureReceiver(this);
+	}
+}
+
+void UMagicLeapTouchpadGesturesComponent::TickComponent(float DeltaTime, enum ELevelTick TickType, FActorComponentTickFunction *ThisTickFunction)
 {
 	Super::TickComponent(DeltaTime, TickType, ThisTickFunction);
 
@@ -78,4 +68,13 @@ void UTouchpadGesturesComponent::TickComponent(float DeltaTime, enum ELevelTick
 
 }
 
+void UMagicLeapTouchpadGesturesComponent::EndPlay(const EEndPlayReason::Type EndPlayReason)
+{
+	TSharedPtr<FMagicLeapController> controller = StaticCastSharedPtr<FMagicLeapController>(IMagicLeapControllerPlugin::Get().GetInputDevice());
+	if (controller.IsValid())
+	{
+		controller->UnregisterTouchpadGestureReceiver(this);
+	}
 
+	Super::EndPlay(EndPlayReason);
+}
diff --git a/Engine/Plugins/Lumin/MagicLeap/Source/MagicLeapController/Public/MagicLeapControllerFunctionLibrary.h b/Engine/Plugins/Lumin/MagicLeap/Source/MagicLeapController/Public/MagicLeapControllerFunctionLibrary.h
--- a/Engine/Plugins/Lumin/MagicLeap/Source/MagicLeapController/Public/MagicLeapControllerFunctionLibrary.h
+++ b/Engine/Plugins/Lumin/MagicLeap/Source/MagicLeapController/Public/MagicLeapControllerFunctionLibrary.h
@@ -24,7 +24,7 @@ public:
 	@return True if the command to play the LED pattern was successfully sent to the controller, false otherwise.
 	*/
 	UFUNCTION(BlueprintCallable, Category = "MotionController|MagicLeap")
-		static bool PlayLEDPattern(FName MotionSource, EMLControllerLEDPattern LEDPattern, EMLControllerLEDColor LEDColor, float DurationInSec);
+	static bool PlayLEDPattern(FName MotionSource, EMagicLeapControllerLEDPattern LEDPattern, EMagicLeapControllerLEDColor LEDColor, float DurationInSec);
 
 	/**
 	Starts a LED feedback effect using the specified pattern on the specified controller.
@@ -37,7 +37,7 @@ public:
 	@return True if the command to play the LED effect was successfully sent to the controller, false otherwise.
 	*/
 	UFUNCTION(BlueprintCallable, Category = "MotionController|MagicLeap")
-		static bool PlayLEDEffect(FName MotionSource, EMLControllerLEDEffect LEDEffect, EMLControllerLEDSpeed LEDSpeed, EMLControllerLEDPattern LEDPattern, EMLControllerLEDColor LEDColor, float DurationInSec);
+	static bool PlayLEDEffect(FName MotionSource, EMagicLeapControllerLEDEffect LEDEffect, EMagicLeapControllerLEDSpeed LEDSpeed, EMagicLeapControllerLEDPattern LEDPattern, EMagicLeapControllerLEDColor LEDColor, float DurationInSec);
 
 	/**
 	Play haptic feedback on the controller.
@@ -47,7 +47,7 @@ public:
 	@return True if the command to play the haptic feedback was successfully sent to the controller, false otherwise.
 	*/
 	UFUNCTION(BlueprintCallable, Category = "MotionController|MagicLeap")
-		static bool PlayHapticPattern(FName MotionSource, EMLControllerHapticPattern HapticPattern, EMLControllerHapticIntensity Intensity);
+	static bool PlayHapticPattern(FName MotionSource, EMagicLeapControllerHapticPattern HapticPattern, EMagicLeapControllerHapticIntensity Intensity);
 
 	/**
 	Set controller tracking mode.
@@ -55,35 +55,35 @@ public:
 	@return True if the command to set the tracking mode was successfully sent to the controller, false otherwise.
 	*/
 	UFUNCTION(BlueprintCallable, Category = "MotionController|MagicLeap")
-		static bool SetControllerTrackingMode(EMLControllerTrackingMode TrackingMode);
+	static bool SetControllerTrackingMode(EMagicLeapControllerTrackingMode TrackingMode);
 
 	/**
 	Get controller tracking mode.
 	@return Controller tracking mode.
 	*/
 	UFUNCTION(BlueprintCallable, Category = "MotionController|MagicLeap")
-		static EMLControllerTrackingMode GetControllerTrackingMode();
+	static EMagicLeapControllerTrackingMode GetControllerTrackingMode();
 
 	/**
 	Get motion source for hand.
 	@return Motion source to which hand is mapped.
 	*/
 	UFUNCTION(BlueprintCallable, Category = "MotionController|MagicLeap")
-		static FName GetMotionSourceForHand(EControllerHand Hand);
+	static FName GetMotionSourceForHand(EControllerHand Hand);
 
 	/**
 	Get motion source for hand.
 	@return Hand to which motion source is mapped, ControllerHand_Count if not mapped.
 	*/
 	UFUNCTION(BlueprintCallable, Category = "MotionController|MagicLeap")
-		static EControllerHand GetHandForMotionSource(FName MotionSource);
+	static EControllerHand GetHandForMotionSource(FName MotionSource);
 
 	/**
 	Set motion source for hand.
 	@return True if successful, false otherwise.
 	*/
 	UFUNCTION(BlueprintCallable, Category = "MotionController|MagicLeap")
-		static bool SetMotionSourceForHand(EControllerHand Hand, FName MotionSource);
+	static bool SetMotionSourceForHand(EControllerHand Hand, FName MotionSource);
 
 	/**
 	Type of ML device bound to the specified motion source
@@ -91,7 +91,15 @@ public:
 	@return Type of ML device which maps to given Unreal controller hand.
 	*/
 	UFUNCTION(BlueprintCallable, BlueprintPure, Category = "MotionController|MagicLeap")
-		static EMLControllerType GetControllerType(FName MotionSource);
+	static EMagicLeapControllerType GetControllerType(FName MotionSource);
+
+	/**
+	Checks if the ML controller (MagicLeapControl0, MagicLeapControl1, MagicLeapMobileApp) of the given motion source is currently connected.
+	@param MotionSource Motion source to query.
+	@return true if the ML controller of the given motion source is connected, false otherwise. Will return false if the motion source is not of an ML controller.
+	*/
+	UFUNCTION(BlueprintCallable, BlueprintPure, Category = "MotionController|MagicLeap")
+	static bool IsMLControllerConnected(FName MotionSource);
 
 	/////////////////////////////////////////////////////////////////////////////////////////////
 	// DEPRECATED FUNCTIONS
@@ -113,7 +121,7 @@ public:
 	  @param Hand Output parameter which is the hand the given index maps to. Valid only if the function returns true.
 	  @return true of the controller index maps to a valid hand, false otherwise
 	*/
-	UFUNCTION(BlueprintCallable, Category = "MotionController|MagicLeap", meta = (DeprecatedFunction, DeprecationMessage = "Use GetHandForMotionSource instead"))
+	UFUNCTION(BlueprintCallable, Category = "MotionController|MagicLeap")
 	static bool GetControllerMapping(int32 ControllerIndex, EControllerHand& Hand);
 
 	/**
@@ -121,8 +129,8 @@ public:
 	@param Hand Controller hand to query.
 	@return Type of ML device which maps to given Unreal controller hand.
 	*/
-	UFUNCTION(BlueprintCallable, BlueprintPure, Category = "MotionController|MagicLeap", meta = (DeprecatedFunction, DeprecationMessage = "Use GetMotionSourceForHand instead"))
-		static EMLControllerType GetMLControllerType(EControllerHand Hand);
+	UFUNCTION(BlueprintCallable, BlueprintPure, Category = "MotionController|MagicLeap")
+	static EMagicLeapControllerType GetMLControllerType(EControllerHand Hand);
 
 	/**
 	  Inverts the controller mapping i.e. keys mapped to left hand controller will now be treated as right hand and vice-versa.
@@ -140,7 +148,7 @@ public:
 	  @return True if the command to play the LED pattern was successfully sent to the controller, false otherwise.
 	*/
 	UFUNCTION(BlueprintCallable, Category = "MotionController|MagicLeap", meta = (DeprecatedFunction, DeprecationMessage = "Use PlayLEDPattern instead"))
-	static bool PlayControllerLED(EControllerHand Hand, EMLControllerLEDPattern LEDPattern, EMLControllerLEDColor LEDColor, float DurationInSec);
+	static bool PlayControllerLED(EControllerHand Hand, EMagicLeapControllerLEDPattern LEDPattern, EMagicLeapControllerLEDColor LEDColor, float DurationInSec);
 
 	/**
 	  Starts a LED feedback effect using the specified pattern on the specified controller.
@@ -153,7 +161,7 @@ public:
 	  @return True if the command to play the LED effect was successfully sent to the controller, false otherwise.
 	*/
 	UFUNCTION(BlueprintCallable, Category = "MotionController|MagicLeap", meta = (DeprecatedFunction, DeprecationMessage = "Use PlayLEDEffect instead"))
-	static bool PlayControllerLEDEffect(EControllerHand Hand, EMLControllerLEDEffect LEDEffect, EMLControllerLEDSpeed LEDSpeed, EMLControllerLEDPattern LEDPattern, EMLControllerLEDColor LEDColor, float DurationInSec);
+	static bool PlayControllerLEDEffect(EControllerHand Hand, EMagicLeapControllerLEDEffect LEDEffect, EMagicLeapControllerLEDSpeed LEDSpeed, EMagicLeapControllerLEDPattern LEDPattern, EMagicLeapControllerLEDColor LEDColor, float DurationInSec);
 
 	/**
 	  Play haptic feedback on the controller.
@@ -163,5 +171,5 @@ public:
 	  @return True if the command to play the haptic feedback was successfully sent to the controller, false otherwise.
 	*/
 	UFUNCTION(BlueprintCallable, Category = "MotionController|MagicLeap", meta = (DeprecatedFunction, DeprecationMessage = "Use PlayHapticPattern instead"))
-	static bool PlayControllerHapticFeedback(EControllerHand Hand, EMLControllerHapticPattern HapticPattern, EMLControllerHapticIntensity Intensity);
+	static bool PlayControllerHapticFeedback(EControllerHand Hand, EMagicLeapControllerHapticPattern HapticPattern, EMagicLeapControllerHapticIntensity Intensity);
 };
diff --git a/Engine/Plugins/Lumin/MagicLeap/Source/MagicLeapController/Public/MagicLeapControllerKeys.h b/Engine/Plugins/Lumin/MagicLeap/Source/MagicLeapController/Public/MagicLeapControllerKeys.h
--- a/Engine/Plugins/Lumin/MagicLeap/Source/MagicLeapController/Public/MagicLeapControllerKeys.h
+++ b/Engine/Plugins/Lumin/MagicLeap/Source/MagicLeapController/Public/MagicLeapControllerKeys.h
@@ -19,32 +19,64 @@ struct FMagicLeapMotionSourceNames
 struct FMagicLeapControllerKeyNames
 {
 	static const FName MotionController_Left_Thumbstick_Z_Name;
-	static const FName Left_MoveButton_Name;
-	static const FName Left_AppButton_Name;
 	static const FName Left_HomeButton_Name;
 
 	static const FName MotionController_Right_Thumbstick_Z_Name;
-	static const FName Right_MoveButton_Name;
-	static const FName Right_AppButton_Name;
 	static const FName Right_HomeButton_Name;
+
+	static const FName TouchpadGesture_Swipe_Up_Name;
+	static const FName TouchpadGesture_Swipe_Down_Name;
+	static const FName TouchpadGesture_Swipe_Right_Name;
+	static const FName TouchpadGesture_Swipe_Left_Name;
+	static const FName TouchpadGesture_RadialScroll_Clockwise_Name;
+	static const FName TouchpadGesture_RadialScroll_CounterClockwise_Name;
+	static const FName TouchpadGesture_Scroll_Up_Name;
+	static const FName TouchpadGesture_Scroll_Down_Name;
+	static const FName TouchpadGesture_Scroll_Right_Name;
+	static const FName TouchpadGesture_Scroll_Left_Name;
+	static const FName TouchpadGesture_Pinch_In_Name;
+	static const FName TouchpadGesture_Pinch_Out_Name;
+	static const FName TouchpadGesture_LongHold_Name;
+	static const FName TouchpadGesture_Tap_Name;
+	static const FName TouchpadGesture_ForceTapDown_Name;
+	static const FName TouchpadGesture_SecondForceDown_Name;
+	static const FName TouchpadGesture_ForceTapUp_Name;
+	static const FName TouchpadGesture_ForceDwell_Name;
+	static const FName TouchpadGesture_Any_Name;
 };
 
-struct FMagicLeapKeys
+struct MAGICLEAPCONTROLLER_API FMagicLeapKeys
 {
 	static const FKey MotionController_Left_Thumbstick_Z;
-	static const FKey Left_MoveButton;
-	static const FKey Left_AppButton;
 	static const FKey Left_HomeButton;
 
 	static const FKey MotionController_Right_Thumbstick_Z;
-	static const FKey Right_MoveButton;
-	static const FKey Right_AppButton;
 	static const FKey Right_HomeButton;
+
+	static const FKey TouchpadGesture_Swipe_Up;
+	static const FKey TouchpadGesture_Swipe_Down;
+	static const FKey TouchpadGesture_Swipe_Right;
+	static const FKey TouchpadGesture_Swipe_Left;
+	static const FKey TouchpadGesture_RadialScroll_Clockwise;
+	static const FKey TouchpadGesture_RadialScroll_CounterClockwise;
+	static const FKey TouchpadGesture_Scroll_Up;
+	static const FKey TouchpadGesture_Scroll_Down;
+	static const FKey TouchpadGesture_Scroll_Right;
+	static const FKey TouchpadGesture_Scroll_Left;
+	static const FKey TouchpadGesture_Pinch_In;
+	static const FKey TouchpadGesture_Pinch_Out;
+	static const FKey TouchpadGesture_LongHold;
+	static const FKey TouchpadGesture_Tap;
+	static const FKey TouchpadGesture_ForceTapDown;
+	static const FKey TouchpadGesture_SecondForceDown;
+	static const FKey TouchpadGesture_ForceTapUp;
+	static const FKey TouchpadGesture_ForceDwell;
+	static const FKey TouchpadGesture_Any;
 };
 
 /** Defines the Magic Leap controller type. */
 UENUM(BlueprintType)
-enum class EMLControllerType : uint8
+enum class EMagicLeapControllerType : uint8
 {
 	None,
 	Device,
@@ -53,7 +85,7 @@ enum class EMLControllerType : uint8
 
 /** LED patterns supported on the controller. */
 UENUM(BlueprintType)
-enum class EMLControllerLEDPattern : uint8
+enum class EMagicLeapControllerLEDPattern : uint8
 {
 	None,
 	Clock01,
@@ -78,7 +110,7 @@ enum class EMLControllerLEDPattern : uint8
 
 /** LED effects supported on the controller. */
 UENUM(BlueprintType)
-enum class EMLControllerLEDEffect : uint8
+enum class EMagicLeapControllerLEDEffect : uint8
 {
 	RotateCW,
 	RotateCCW,
@@ -90,7 +122,7 @@ enum class EMLControllerLEDEffect : uint8
 
 /** LED colors supported on the controller. */
 UENUM(BlueprintType)
-enum class EMLControllerLEDColor : uint8
+enum class EMagicLeapControllerLEDColor : uint8
 {
 	BrightMissionRed,
 	PastelMissionRed,
@@ -112,7 +144,7 @@ enum class EMLControllerLEDColor : uint8
 
 /** LED speeds supported on the controller. */
 UENUM(BlueprintType)
-enum class EMLControllerLEDSpeed : uint8
+enum class EMagicLeapControllerLEDSpeed : uint8
 {
 	Slow,
 	Medium,
@@ -121,7 +153,7 @@ enum class EMLControllerLEDSpeed : uint8
 
 /** Haptic patterns supported on the controller. */
 UENUM(BlueprintType)
-enum class EMLControllerHapticPattern : uint8
+enum class EMagicLeapControllerHapticPattern : uint8
 {
 	None,
 	Click,
@@ -137,7 +169,7 @@ enum class EMLControllerHapticPattern : uint8
 
 /** Haptic intesities supported on the controller. */
 UENUM(BlueprintType)
-enum class EMLControllerHapticIntensity : uint8
+enum class EMagicLeapControllerHapticIntensity : uint8
 {
 	Low,
 	Medium,
@@ -146,7 +178,7 @@ enum class EMLControllerHapticIntensity : uint8
 
 /** Tracking modes provided by Magic Leap. */
 UENUM(BlueprintType)
-enum class EMLControllerTrackingMode : uint8
+enum class EMagicLeapControllerTrackingMode : uint8
 {
 	InputService,
 	CoordinateFrameUID,
diff --git a/Engine/Plugins/Lumin/MagicLeap/Source/MagicLeapController/Public/TouchpadGesturesComponent.h b/Engine/Plugins/Lumin/MagicLeap/Source/MagicLeapController/Public/TouchpadGesturesComponent.h
--- a/Engine/Plugins/Lumin/MagicLeap/Source/MagicLeapController/Public/TouchpadGesturesComponent.h
+++ b/Engine/Plugins/Lumin/MagicLeap/Source/MagicLeapController/Public/TouchpadGesturesComponent.h
@@ -7,7 +7,7 @@
 #include "TouchpadGesturesComponent.generated.h"
 
 /** Recognized touchpad gesture types. */
-UENUM(BlueprintType)
+UENUM(BlueprintType, meta = (DeprecationMessage = "Use Touchpad Gesture Events instead."))
 enum class EMagicLeapTouchpadGestureType : uint8
 {
 	None,
@@ -24,7 +24,7 @@ enum class EMagicLeapTouchpadGestureType : uint8
 };
 
 /** Direction of touchpad gesture. */
-UENUM(BlueprintType)
+UENUM(BlueprintType, meta = (DeprecationMessage = "Use Touchpad Gesture Events instead."))
 enum class EMagicLeapTouchpadGestureDirection : uint8
 {
 	None,
@@ -39,7 +39,7 @@ enum class EMagicLeapTouchpadGestureDirection : uint8
 };
 
 /** Information about a recognized touchpad gesture. */
-USTRUCT(BlueprintType)
+USTRUCT(BlueprintType, meta = (DeprecationMessage="Use Touchpad Gesture Events instead."))
 struct FMagicLeapTouchpadGesture
 {
 	GENERATED_BODY()
@@ -50,22 +50,22 @@ public:
 	EControllerHand Hand;
 
 	/** Motion source on which the gesture was performed. */
-	UPROPERTY(BlueprintReadOnly, Category = "TouchpadGesture|MagicLeap")
+	UPROPERTY(BlueprintReadOnly, Category = "TouchpadGesture|MagicLeap", meta = (DeprecatedProperty, DeprecationMessage = "Use Touchpad Gesture Events instead."))
 	FName MotionSource;
 
 	/** Type of gesture. */
-	UPROPERTY(BlueprintReadOnly, Category = "TouchpadGesture|MagicLeap")
+	UPROPERTY(BlueprintReadOnly, Category = "TouchpadGesture|MagicLeap", meta = (DeprecatedProperty, DeprecationMessage = "Use Touchpad Gesture Events instead."))
 	EMagicLeapTouchpadGestureType Type;
 
 	/** Direction of gesture */
-	UPROPERTY(BlueprintReadOnly, Category = "TouchpadGesture|MagicLeap")
+	UPROPERTY(BlueprintReadOnly, Category = "TouchpadGesture|MagicLeap", meta = (DeprecatedProperty, DeprecationMessage = "Use Touchpad Gesture Events instead."))
 	EMagicLeapTouchpadGestureDirection Direction;
 
 	/** 
 	  Gesture position (x,y) and force (z).
 	  Position is in the [-1.0,1.0] range and force is in the [0.0,1.0] range.
 	*/
-	UPROPERTY(BlueprintReadOnly, Category = "TouchpadGesture|MagicLeap")
+	UPROPERTY(BlueprintReadOnly, Category = "TouchpadGesture|MagicLeap", meta = (DeprecatedProperty, DeprecationMessage = "Use Touchpad Gesture Events instead."))
 	FVector PositionAndForce;
 
 	/**
@@ -77,7 +77,7 @@ public:
         of [-1.0,1.0] so touchpad distance has a range of [0.0,2.0]; this value
         will be in touchpad distance per second.
 	 */
-	UPROPERTY(BlueprintReadOnly, Category = "TouchpadGesture|MagicLeap")
+	UPROPERTY(BlueprintReadOnly, Category = "TouchpadGesture|MagicLeap", meta = (DeprecatedProperty, DeprecationMessage = "Use Touchpad Gesture Events instead."))
 	float Speed;
 
 	/**
@@ -86,7 +86,7 @@ public:
       distance. The touchpad is defined as having extents of [-1.0,1.0] so
       this distance has a range of [0.0,2.0].
 	 */
-	UPROPERTY(BlueprintReadOnly, Category = "TouchpadGesture|MagicLeap")
+	UPROPERTY(BlueprintReadOnly, Category = "TouchpadGesture|MagicLeap", meta = (DeprecatedProperty, DeprecationMessage = "Use Touchpad Gesture Events instead."))
 	float Distance;
 
 	/**
@@ -94,7 +94,7 @@ public:
       distance. The touchpad is defined as having extents of [-1.0,1.0] so
       this distance has a range of [0.0,2.0].
 	 */
-	UPROPERTY(BlueprintReadOnly, Category = "TouchpadGesture|MagicLeap")
+	UPROPERTY(BlueprintReadOnly, Category = "TouchpadGesture|MagicLeap", meta = (DeprecatedProperty, DeprecationMessage = "Use Touchpad Gesture Events instead."))
 	float FingerGap;
 
 	/**
@@ -102,11 +102,11 @@ public:
       is defined as having extents of [-1.0,1.0] so this radius has a range
       of [0.0,2.0].
 	 */
-	UPROPERTY(BlueprintReadOnly, Category = "TouchpadGesture|MagicLeap")
+	UPROPERTY(BlueprintReadOnly, Category = "TouchpadGesture|MagicLeap", meta = (DeprecatedProperty, DeprecationMessage = "Use Touchpad Gesture Events instead."))
 	float Radius;
 
 	/** Angle from the center of the touchpad to the finger. */
-	UPROPERTY(BlueprintReadOnly, Category = "TouchpadGesture|MagicLeap")
+	UPROPERTY(BlueprintReadOnly, Category = "TouchpadGesture|MagicLeap", meta = (DeprecatedProperty, DeprecationMessage = "Use Touchpad Gesture Events instead."))
 	float Angle;
 };
 
@@ -118,28 +118,27 @@ public:
 	virtual void OnTouchpadGestureEndCallback(const FMagicLeapTouchpadGesture& GestureData) = 0;
 };
 
-/** Delegates touchpad gesture events for the Magic Leap Controller & MLMA */
-UCLASS(ClassGroup = MagicLeap, BlueprintType, Blueprintable, EditInlineNew, meta = (BlueprintSpawnableComponent))
-class MAGICLEAPCONTROLLER_API UTouchpadGesturesComponent : public UActorComponent, public IMagicLeapTouchpadGestures
+/** DEPRECATED use Touchpad Gesture Events instead. - Delegates touchpad gesture events for the Magic Leap Controller & MLMA */
+UCLASS(ClassGroup = MagicLeap, BlueprintType, Blueprintable, EditInlineNew, meta = (Deprecated, DeprecationMessage="Use Touchpad Gesture Events instead.", BlueprintSpawnableComponent))
+class MAGICLEAPCONTROLLER_API UMagicLeapTouchpadGesturesComponent : public UActorComponent, public IMagicLeapTouchpadGestures
 {
 	GENERATED_BODY()
 
 public:
-	UTouchpadGesturesComponent();
-	virtual ~UTouchpadGesturesComponent();
+	UMagicLeapTouchpadGesturesComponent();
 
 	DECLARE_DYNAMIC_MULTICAST_DELEGATE_OneParam(FTouchpadGestureEvent, const FMagicLeapTouchpadGesture&, GestureData);
 
-	/** Event called when a touchpad gesture starts. Provides all the meta data about the given gestures. */
-	UPROPERTY(BlueprintAssignable)
+	/** DEPRECATED use Touchpad Gesture Events instead. - Event called when a touchpad gesture starts. Provides all the meta data about the given gestures. */
+	UPROPERTY(BlueprintAssignable, meta = (DeprecatedProperty, DeprecationMessage="Use Touchpad Gesture Events instead."))
 	FTouchpadGestureEvent OnTouchpadGestureStart;
 
-	/** Event called when a touchpad gesture continues. Provides all the meta data about the given gestures. */
-	UPROPERTY(BlueprintAssignable)
+	/** DEPRECATED use Touchpad Gesture Events instead. - Event called when a touchpad gesture continues. Provides all the meta data about the given gestures. */
+	UPROPERTY(BlueprintAssignable, meta = (DeprecatedProperty, DeprecationMessage="Use Touchpad Gesture Events instead."))
 	FTouchpadGestureEvent OnTouchpadGestureContinue;
 
-	/** Event called when a touchpad gesture ends. Provides all the meta data about the given gestures. */
-	UPROPERTY(BlueprintAssignable)
+	/** DEPRECATED use Touchpad Gesture Events instead. - Event called when a touchpad gesture ends. Provides all the meta data about the given gestures. */
+	UPROPERTY(BlueprintAssignable, meta = (DeprecatedProperty, DeprecationMessage="Use Touchpad Gesture Events instead."))
 	FTouchpadGestureEvent OnTouchpadGestureEnd;
 
 	/** IMagicLeapTouchpadGestures interface */
@@ -147,7 +146,9 @@ public:
 	virtual void OnTouchpadGestureContinueCallback(const FMagicLeapTouchpadGesture& GestureData) override;
 	virtual void OnTouchpadGestureEndCallback(const FMagicLeapTouchpadGesture& GestureData) override;
 
+	virtual void BeginPlay() override;
 	virtual void TickComponent(float DeltaTime, enum ELevelTick TickType, FActorComponentTickFunction *ThisTickFunction) override;
+	virtual void EndPlay(const EEndPlayReason::Type EndPlayReason) override;
 
 private:
 	FCriticalSection CriticalSection;
diff --git a/Engine/Plugins/Lumin/MagicLeap/Source/MagicLeapEyeTracker/MagicLeapEyeTracker.Build.cs b/Engine/Plugins/Lumin/MagicLeap/Source/MagicLeapEyeTracker/MagicLeapEyeTracker.Build.cs
--- a/Engine/Plugins/Lumin/MagicLeap/Source/MagicLeapEyeTracker/MagicLeapEyeTracker.Build.cs
+++ b/Engine/Plugins/Lumin/MagicLeap/Source/MagicLeapEyeTracker/MagicLeapEyeTracker.Build.cs
@@ -31,10 +31,6 @@ namespace UnrealBuildTool.Rules
 					"MLSDK"
 				}
 			);
-
-			// This is not ideal but needs to be done in order to expose the private MagicLeapHMD header to this module.
-			PrivateIncludePaths.Add(Path.Combine(new string[] { ModuleDirectory, "..", "MagicLeap", "Private" }));
-			//PublicIncludePaths.AddRange(new string[] { "Private", "Public"});
 		}
 	}
 }
diff --git a/Engine/Plugins/Lumin/MagicLeap/Source/MagicLeapEyeTracker/Private/MagicLeapEyeTrackerModule.cpp b/Engine/Plugins/Lumin/MagicLeap/Source/MagicLeapEyeTracker/Private/MagicLeapEyeTrackerModule.cpp
--- a/Engine/Plugins/Lumin/MagicLeap/Source/MagicLeapEyeTracker/Private/MagicLeapEyeTrackerModule.cpp
+++ b/Engine/Plugins/Lumin/MagicLeap/Source/MagicLeapEyeTracker/Private/MagicLeapEyeTrackerModule.cpp
@@ -131,16 +131,6 @@ void FMagicLeapEyeTracker::SetEyeTrackedPlayer(APlayerController* PlayerControll
 	}
 }
 
-bool FMagicLeapEyeTracker::IsEyeTrackerCalibrated() const
-{
-	if (VREyeTracker)
-	{
-		return VREyeTracker->IsEyeTrackerCalibrated();
-	}
-
-	return false;
-}
-
 bool FMagicLeapEyeTracker::GetEyeBlinkState(FMagicLeapEyeBlinkState& BlinkState) const
 {
 	if (VREyeTracker)
@@ -180,10 +170,7 @@ bool FMagicLeapEyeTracker::GetFixationComfort(FMagicLeapFixationComfort& Fixatio
 /* FMagicLeapEyeTrackerModule                                           */
 /************************************************************************/
 FMagicLeapEyeTrackerModule::FMagicLeapEyeTrackerModule()
-	: IMagicLeapModule("MagicLeapEyeTracker")
-{
-
-}
+{}
 
 void FMagicLeapEyeTrackerModule::StartupModule()
 {
@@ -198,14 +185,6 @@ void FMagicLeapEyeTrackerModule::ShutdownModule()
 	FCoreUObjectDelegates::PreLoadMap.Remove(OnPreLoadMapHandle);
 }
 
-void FMagicLeapEyeTrackerModule::Disable()
-{
-	if (MagicLeapEyeTracker.IsValid())
-	{
-		MagicLeapEyeTracker->Destroy();
-	}
-}
-
 TSharedPtr<class IEyeTracker, ESPMode::ThreadSafe> FMagicLeapEyeTrackerModule::CreateEyeTracker()
 {
 	return MagicLeapEyeTracker;
@@ -238,18 +217,6 @@ bool FMagicLeapEyeTrackerModule::IsEyeTrackerConnected() const
 	return false;
 }
 
-bool UMagicLeapEyeTrackerFunctionLibrary::IsEyeTrackerCalibrated()
-{
-	// TODO: Don't do this. Use StaticCastSharedPtr().
-	FMagicLeapEyeTracker* ET = GEngine ? static_cast<FMagicLeapEyeTracker*>(GEngine->EyeTrackingDevice.Get()) : nullptr;
-	if (ET)
-	{
-		return ET->IsEyeTrackerCalibrated();
-	}
-
-	return false;
-}
-
 bool UMagicLeapEyeTrackerFunctionLibrary::GetEyeBlinkState(FMagicLeapEyeBlinkState& BlinkState)
 {
 	// TODO: Don't do this. Use StaticCastSharedPtr().
diff --git a/Engine/Plugins/Lumin/MagicLeap/Source/MagicLeapEyeTracker/Private/MagicLeapVREyeTracker.cpp b/Engine/Plugins/Lumin/MagicLeap/Source/MagicLeapEyeTracker/Private/MagicLeapVREyeTracker.cpp
--- a/Engine/Plugins/Lumin/MagicLeap/Source/MagicLeapEyeTracker/Private/MagicLeapVREyeTracker.cpp
+++ b/Engine/Plugins/Lumin/MagicLeap/Source/MagicLeapEyeTracker/Private/MagicLeapVREyeTracker.cpp
@@ -11,19 +11,11 @@
 #include "Engine/LocalPlayer.h"
 #include "SceneView.h"
 #include "IMagicLeapPlugin.h"
-#include "AppFramework.h"
-#include "MagicLeapHMD.h"
 #include "UnrealEngine.h"
-#include "MagicLeapSDKDetection.h"
 #include "HeadMountedDisplayFunctionLibrary.h"
-#include "MagicLeapPluginUtil.h" // for ML_INCLUDES_START/END
 #include "IMagicLeapPlugin.h"
-
-#if WITH_MLSDK
-ML_INCLUDES_START
-#include <ml_eye_tracking.h>
-ML_INCLUDES_END
-#endif //WITH_MLSDK
+#include "Lumin/CAPIShims/LuminAPIEyeTracking.h"
+#include "MagicLeapCFUID.h"
 
 #if WITH_MLSDK
 EMagicLeapEyeTrackingCalibrationStatus MLToUnrealEyeCalibrationStatus(MLEyeTrackingCalibrationStatus InStatus)
@@ -49,10 +41,17 @@ FMagicLeapVREyeTracker::FMagicLeapVREyeTracker()
 , EyeTrackingHandle(ML_INVALID_HANDLE)
 #endif //WITH_MLSDK
 {
+	IMagicLeapPlugin::Get().RegisterMagicLeapTrackerEntity(this);
 	SetDefaultDataValues();
 }
 
 FMagicLeapVREyeTracker::~FMagicLeapVREyeTracker()
+{
+	DestroyEntityTracker();
+	IMagicLeapPlugin::Get().UnregisterMagicLeapTrackerEntity(this);
+}
+
+void FMagicLeapVREyeTracker::DestroyEntityTracker()
 {
 #if WITH_MLSDK
 	if (EyeTrackingHandle != ML_INVALID_HANDLE)
@@ -64,9 +63,9 @@ FMagicLeapVREyeTracker::~FMagicLeapVREyeTracker()
 		EyeTrackingHandle = ML_INVALID_HANDLE;
 	}
 #endif //WITH_MLSDK
+	bInitialized = false;
 }
 
-
 void FMagicLeapVREyeTracker::SetDefaultDataValues()
 {
 	FMemory::Memzero(UnfilteredEyeTrackingData);
@@ -112,15 +111,15 @@ bool FMagicLeapVREyeTracker::Tick(float DeltaTime)
 			FDateTime Now = FDateTime::UtcNow();
 			UnfilteredEyeTrackingData.TimeStamp = Now;
 
-			const FAppFramework& AppFramework = static_cast<FMagicLeapHMD*>(GEngine->XRSystem->GetHMDDevice())->GetAppFramework();
-			if (AppFramework.IsInitialized())
+			IMagicLeapPlugin& MLPlugin = IMagicLeapPlugin::Get();
+			if (MLPlugin.IsPerceptionEnabled())
 			{
-				const FTransform TrackingToWorld = UHeadMountedDisplayFunctionLibrary::GetTrackingToWorldTransform(GWorld);
+				const FTransform TrackingToWorld = UHeadMountedDisplayFunctionLibrary::GetTrackingToWorldTransform(nullptr);
 
 				//harvest the 3 transforms
-				EFailReason FailReason;
+				EMagicLeapTransformFailReason FailReason;
 				FTransform FixationTransform;
-				if (AppFramework.GetTransform(EyeTrackingStaticData.fixation, FixationTransform, FailReason))
+				if (MLPlugin.GetTransform(EyeTrackingStaticData.fixation, FixationTransform, FailReason))
 				{
 					FixationTransform = FixationTransform * TrackingToWorld;
 					//get focal point
@@ -129,7 +128,7 @@ bool FMagicLeapVREyeTracker::Tick(float DeltaTime)
 
 				bool bLeftTransformValid = false;
 				FTransform LeftCenterTransform;
-				if (AppFramework.GetTransform(EyeTrackingStaticData.left_center, LeftCenterTransform, FailReason))
+				if (MLPlugin.GetTransform(EyeTrackingStaticData.left_center, LeftCenterTransform, FailReason))
 				{
 					LeftCenterTransform = LeftCenterTransform * TrackingToWorld;
 					bLeftTransformValid = true;
@@ -137,7 +136,7 @@ bool FMagicLeapVREyeTracker::Tick(float DeltaTime)
 
 				bool bRightTransformValid = false;
 				FTransform RightCenterTransform;
-				if (AppFramework.GetTransform(EyeTrackingStaticData.right_center, RightCenterTransform, FailReason))
+				if (MLPlugin.GetTransform(EyeTrackingStaticData.right_center, RightCenterTransform, FailReason))
 				{
 					RightCenterTransform = RightCenterTransform * TrackingToWorld;
 					bRightTransformValid = true;
@@ -179,12 +178,11 @@ bool FMagicLeapVREyeTracker::Tick(float DeltaTime)
 	{
 		if (!bInitialized && IMagicLeapPlugin::Get().IsMagicLeapHMDValid())
 		{
-#if !PLATFORM_MAC
-			if (static_cast<FMagicLeapHMD*>(GEngine->XRSystem->GetHMDDevice())->IsPerceptionEnabled())
+			if (IMagicLeapPlugin::Get().IsPerceptionEnabled())
 			{
 				//keep trying until we are successful in creating one
 				MLResult CreateResult = MLResult_UnspecifiedFailure;
-				ML_FUNCTION_WRAPPER(CreateResult = MLEyeTrackingCreate(&EyeTrackingHandle));
+				CreateResult = MLEyeTrackingCreate(&EyeTrackingHandle);
 				bInitialized = CreateResult == MLResult_Ok && MLHandleIsValid(EyeTrackingHandle);
 				if (bInitialized)
 				{
@@ -196,7 +194,6 @@ bool FMagicLeapVREyeTracker::Tick(float DeltaTime)
 					}
 				}
 			}
-#endif //!PLATFORM_MAC
 		}
 	}
 #endif //WITH_MLSDK
@@ -220,11 +217,6 @@ EMagicLeapEyeTrackingStatus FMagicLeapVREyeTracker::GetEyeTrackingStatus()
 	return EyeTrackingStatus;
 }
 
-bool FMagicLeapVREyeTracker::IsEyeTrackerCalibrated() const
-{
-	return GetCalibrationStatus() == EMagicLeapEyeTrackingCalibrationStatus::None ? false : true;
-}
-
 EMagicLeapEyeTrackingCalibrationStatus FMagicLeapVREyeTracker::GetCalibrationStatus() const
 {
 	return EyeCalibrationStatus;
diff --git a/Engine/Plugins/Lumin/MagicLeap/Source/MagicLeapEyeTracker/Private/MagicLeapVREyeTracker.h b/Engine/Plugins/Lumin/MagicLeap/Source/MagicLeapEyeTracker/Private/MagicLeapVREyeTracker.h
--- a/Engine/Plugins/Lumin/MagicLeap/Source/MagicLeapEyeTracker/Private/MagicLeapVREyeTracker.h
+++ b/Engine/Plugins/Lumin/MagicLeap/Source/MagicLeapEyeTracker/Private/MagicLeapVREyeTracker.h
@@ -9,21 +9,18 @@
 #include "MagicLeapEyeTrackerTypes.h"
 #include "Containers/Ticker.h"
 #include "GameFramework/HUD.h"
-#include "MagicLeapPluginUtil.h" // for ML_INCLUDES_START/END
+#include "Lumin/CAPIShims/LuminAPIEyeTracking.h"
+#include "IMagicLeapTrackerEntity.h"
 
-#if WITH_MLSDK
-ML_INCLUDES_START
-#include <ml_api.h>
-#include <ml_eye_tracking.h>
-ML_INCLUDES_END
-#endif //WITH_MLSDK
-
-class FMagicLeapVREyeTracker : public IMagicLeapVREyeTracker, public FTickerObjectBase
+class FMagicLeapVREyeTracker : public IMagicLeapVREyeTracker, public FTickerObjectBase, public IMagicLeapTrackerEntity
 {
 public:
 	FMagicLeapVREyeTracker();
 	virtual ~FMagicLeapVREyeTracker();
 
+	/** IMagicLeapTrackerEntity interface */
+	virtual void DestroyEntityTracker() override;
+
 	void SetDefaultDataValues();
 
 	void SetActivePlayerController(APlayerController* NewActivePlayerController);
@@ -33,7 +30,6 @@ public:
 
 	void DrawDebug(AHUD* HUD, UCanvas* Canvas, const FDebugDisplayInfo& DisplayInfo, float& YL, float& YPos);
 
-	bool IsEyeTrackerCalibrated() const;
 	EMagicLeapEyeTrackingCalibrationStatus GetCalibrationStatus() const;
 
 public:
diff --git a/Engine/Plugins/Lumin/MagicLeap/Source/MagicLeapEyeTracker/Public/MagicLeapEyeTrackerModule.h b/Engine/Plugins/Lumin/MagicLeap/Source/MagicLeapEyeTracker/Public/MagicLeapEyeTrackerModule.h
--- a/Engine/Plugins/Lumin/MagicLeap/Source/MagicLeapEyeTracker/Public/MagicLeapEyeTrackerModule.h
+++ b/Engine/Plugins/Lumin/MagicLeap/Source/MagicLeapEyeTracker/Public/MagicLeapEyeTrackerModule.h
@@ -9,7 +9,6 @@
 #include "GameFramework/HUD.h"
 #include "Kismet/BlueprintFunctionLibrary.h"
 #include "MagicLeapEyeTrackerTypes.h"
-#include "IMagicLeapModule.h"
 #include "MagicLeapEyeTrackerModule.generated.h"
 
 class FMagicLeapVREyeTracker;
@@ -107,7 +106,6 @@ public:
 	virtual EEyeTrackerStatus GetEyeTrackerStatus() const override;
 	virtual bool IsStereoGazeDataAvailable() const override;
 
-	bool IsEyeTrackerCalibrated() const;
 	bool GetEyeBlinkState(FMagicLeapEyeBlinkState& BlinkState) const;
 	bool GetFixationComfort(FMagicLeapFixationComfort& FixationComfort) const;
 
@@ -122,7 +120,7 @@ private:
 	FMagicLeapVREyeTracker* VREyeTracker;
 };
 
-class FMagicLeapEyeTrackerModule : public IMagicLeapEyeTrackerModule, public IMagicLeapModule
+class FMagicLeapEyeTrackerModule : public IMagicLeapEyeTrackerModule
 {
 	/************************************************************************/
 	/* IInputDeviceModule                                                   */
@@ -131,7 +129,6 @@ public:
 	FMagicLeapEyeTrackerModule();
 	virtual void StartupModule() override;
 	virtual void ShutdownModule() override;
-	virtual void Disable() override;
 	virtual TSharedPtr< class IEyeTracker, ESPMode::ThreadSafe > CreateEyeTracker() override;
 
 	/************************************************************************/
@@ -159,17 +156,10 @@ class MAGICLEAPEYETRACKER_API UMagicLeapEyeTrackerFunctionLibrary : public UBlue
 	GENERATED_BODY()
 
 public:
-	/** 
-	  False if the calibration status is none, otherwise returns true, even with a bad calibration.
-	  If not, user should be advised to run the Eye Calibrator app on the device.
-	*/
-	UFUNCTION(BlueprintCallable, BlueprintPure, meta = (DeprecatedFunction, DeprecationMessage = "Please use GetCalibrationStatus instead"), Category = "Eye Tracking|MagicLeap")
-	static bool IsEyeTrackerCalibrated();
-
 	UFUNCTION(BlueprintCallable, BlueprintPure, Category = "Eye Tracking|MagicLeap")
 	static bool GetEyeBlinkState(FMagicLeapEyeBlinkState& BlinkState);
 
-	UFUNCTION(BlueprintCallable, BlueprintPure, Category = "Eye Tracking|MagicLeap")
+	UFUNCTION(BlueprintCallable, BlueprintPure, Category = "Eye Tracking|MagicLeap", meta = (DeprecatedFunction, DeprecationMessage = "Eye fixation comfort api is no longer available"))
 	static bool GetFixationComfort(FMagicLeapFixationComfort& FixationComfort);
 
 	UFUNCTION(BlueprintCallable, BlueprintPure, Category = "EyeTracking|MagicLeap")
diff --git a/Engine/Plugins/Lumin/MagicLeap/Source/MagicLeapHandMeshing/MagicLeapHandMeshing.Build.cs b/Engine/Plugins/Lumin/MagicLeap/Source/MagicLeapHandMeshing/MagicLeapHandMeshing.Build.cs
new file mode 100644
--- /dev/null
+++ b/Engine/Plugins/Lumin/MagicLeap/Source/MagicLeapHandMeshing/MagicLeapHandMeshing.Build.cs
@@ -0,0 +1,31 @@
+// Copyright 1998-2018 Epic Games, Inc. All Rights Reserved.
+
+using System.IO;
+
+namespace UnrealBuildTool.Rules
+{
+	public class MagicLeapHandMeshing : ModuleRules
+	{
+		public MagicLeapHandMeshing( ReadOnlyTargetRules Target ) : base(Target)
+		{
+			PublicDependencyModuleNames.AddRange(
+				new string[]
+				{
+					"Core",
+					"CoreUObject",
+					"Engine"
+				}
+			);
+
+			PrivateDependencyModuleNames.AddRange(
+				new string[]
+				{
+					"LuminRuntimeSettings",
+					"MagicLeap",
+					"HeadMountedDisplay",
+					"MRMesh"
+				}
+			);
+		}
+	}
+}
diff --git a/Engine/Plugins/Lumin/MagicLeap/Source/MagicLeapHandMeshing/Private/MagicLeapHandMeshingComponent.cpp b/Engine/Plugins/Lumin/MagicLeap/Source/MagicLeapHandMeshing/Private/MagicLeapHandMeshingComponent.cpp
new file mode 100644
--- /dev/null
+++ b/Engine/Plugins/Lumin/MagicLeap/Source/MagicLeapHandMeshing/Private/MagicLeapHandMeshingComponent.cpp
@@ -0,0 +1,26 @@
+// Copyright 1998-2018 Epic Games, Inc. All Rights Reserved.
+
+#include "MagicLeapHandMeshingComponent.h"
+#include "MagicLeapHandMeshingModule.h"
+
+void UMagicLeapHandMeshingComponent::EndPlay(const EEndPlayReason::Type EndPlayReason)
+{
+	GetMagicLeapHandMeshingModule().DestroyClient();
+	Super::EndPlay(EndPlayReason);
+}
+
+bool UMagicLeapHandMeshingComponent::ConnectMRMesh(UMRMeshComponent* InMRMeshPtr)
+{
+	FMagicLeapHandMeshingModule& Module = GetMagicLeapHandMeshingModule();
+	if (!Module.HasClient())
+	{
+		Module.CreateClient();
+	}
+
+	return Module.ConnectMRMesh(InMRMeshPtr);
+}
+
+bool UMagicLeapHandMeshingComponent::DisconnectMRMesh(class UMRMeshComponent* InMRMeshPtr)
+{
+	return GetMagicLeapHandMeshingModule().DisconnectMRMesh(InMRMeshPtr);
+}
diff --git a/Engine/Plugins/Lumin/MagicLeap/Source/MagicLeapHandMeshing/Private/MagicLeapHandMeshingFunctionLibrary.cpp b/Engine/Plugins/Lumin/MagicLeap/Source/MagicLeapHandMeshing/Private/MagicLeapHandMeshingFunctionLibrary.cpp
new file mode 100644
--- /dev/null
+++ b/Engine/Plugins/Lumin/MagicLeap/Source/MagicLeapHandMeshing/Private/MagicLeapHandMeshingFunctionLibrary.cpp
@@ -0,0 +1,24 @@
+// Copyright 1998-2018 Epic Games, Inc. All Rights Reserved.
+
+#include "MagicLeapHandMeshingFunctionLibrary.h"
+#include "MagicLeapHandMeshingModule.h"
+
+bool UMagicLeapHandMeshingFunctionLibrary::CreateClient()
+{
+	return GetMagicLeapHandMeshingModule().CreateClient();
+}
+
+bool UMagicLeapHandMeshingFunctionLibrary::DestroyClient()
+{
+	return GetMagicLeapHandMeshingModule().DestroyClient();
+}
+
+bool UMagicLeapHandMeshingFunctionLibrary::ConnectMRMesh(UMRMeshComponent* InMRMeshPtr)
+{
+	return GetMagicLeapHandMeshingModule().ConnectMRMesh(InMRMeshPtr);
+}
+
+bool UMagicLeapHandMeshingFunctionLibrary::DisconnectMRMesh(class UMRMeshComponent* InMRMeshPtr)
+{
+	return GetMagicLeapHandMeshingModule().DisconnectMRMesh(InMRMeshPtr);
+}
diff --git a/Engine/Plugins/Lumin/MagicLeap/Source/MagicLeapHandMeshing/Private/MagicLeapHandMeshingModule.cpp b/Engine/Plugins/Lumin/MagicLeap/Source/MagicLeapHandMeshing/Private/MagicLeapHandMeshingModule.cpp
new file mode 100644
--- /dev/null
+++ b/Engine/Plugins/Lumin/MagicLeap/Source/MagicLeapHandMeshing/Private/MagicLeapHandMeshingModule.cpp
@@ -0,0 +1,58 @@
+// Copyright 1998-2018 Epic Games, Inc. All Rights Reserved.
+
+#include "MagicLeapHandMeshingModule.h"
+#include "GameFramework/Actor.h"
+#include "EngineUtils.h"
+
+DEFINE_LOG_CATEGORY(LogMagicLeapHandMeshing);
+
+void FMagicLeapHandMeshingModule::StartupModule()
+{
+	IMagicLeapHandMeshingModule::StartupModule();
+	TickDelegate = FTickerDelegate::CreateRaw(this, &FMagicLeapHandMeshingModule::Tick);
+	TickDelegateHandle = FTicker::GetCoreTicker().AddTicker(TickDelegate);
+}
+
+void FMagicLeapHandMeshingModule::ShutdownModule()
+{
+	FTicker::GetCoreTicker().RemoveTicker(TickDelegateHandle);
+	IModuleInterface::ShutdownModule();
+}
+
+bool FMagicLeapHandMeshingModule::Tick(float DeltaTime)
+{
+	MeshTracker.Update();
+	return true;
+}
+
+bool FMagicLeapHandMeshingModule::CreateClient()
+{
+	return MeshTracker.Create();
+}
+
+bool FMagicLeapHandMeshingModule::DestroyClient()
+{
+	return MeshTracker.Destroy();
+}
+
+bool FMagicLeapHandMeshingModule::ConnectMRMesh(UMRMeshComponent* InMRMeshPtr)
+{
+	return MeshTracker.ConnectMRMesh(InMRMeshPtr);
+}
+
+bool FMagicLeapHandMeshingModule::DisconnectMRMesh(class UMRMeshComponent* InMRMeshPtr)
+{
+	return MeshTracker.DisconnectMRMesh(InMRMeshPtr);
+}
+
+bool FMagicLeapHandMeshingModule::HasMRMesh() const
+{
+	return MeshTracker.HasMRMesh();
+}
+
+bool FMagicLeapHandMeshingModule::HasClient() const
+{
+	return MeshTracker.HasClient();
+}
+
+IMPLEMENT_MODULE(FMagicLeapHandMeshingModule, MagicLeapHandMeshing);
diff --git a/Engine/Plugins/Lumin/MagicLeap/Source/MagicLeapHandMeshing/Private/MagicLeapHandMeshingModule.h b/Engine/Plugins/Lumin/MagicLeap/Source/MagicLeapHandMeshing/Private/MagicLeapHandMeshingModule.h
new file mode 100644
--- /dev/null
+++ b/Engine/Plugins/Lumin/MagicLeap/Source/MagicLeapHandMeshing/Private/MagicLeapHandMeshingModule.h
@@ -0,0 +1,34 @@
+// Copyright 1998-2018 Epic Games, Inc. All Rights Reserved.
+
+#pragma once
+
+#include "CoreMinimal.h"
+#include "IMagicLeapPlugin.h"
+#include "IMagicLeapHandMeshingModule.h"
+#include "MagicLeapMeshTracker.h"
+
+DECLARE_LOG_CATEGORY_EXTERN(LogMagicLeapHandMeshing, Verbose, All);
+
+class FMagicLeapHandMeshingModule : public IMagicLeapHandMeshingModule
+{
+public:
+	void StartupModule() override;
+	void ShutdownModule() override;
+	bool Tick(float DeltaTime) override;
+	bool CreateClient();
+	bool DestroyClient();
+	bool ConnectMRMesh(UMRMeshComponent* InMRMeshPtr);
+	bool DisconnectMRMesh(class UMRMeshComponent* InMRMeshPtr);
+	bool HasMRMesh() const;
+	bool HasClient() const;
+
+private:
+	FTickerDelegate TickDelegate;
+	FDelegateHandle TickDelegateHandle;
+	FMagicLeapMeshTracker MeshTracker;
+};
+
+inline FMagicLeapHandMeshingModule& GetMagicLeapHandMeshingModule()
+{
+	return FModuleManager::Get().GetModuleChecked<FMagicLeapHandMeshingModule>("MagicLeapHandMeshing");
+}
diff --git a/Engine/Plugins/Lumin/MagicLeap/Source/MagicLeapHandMeshing/Private/MagicLeapMeshTracker.cpp b/Engine/Plugins/Lumin/MagicLeap/Source/MagicLeapHandMeshing/Private/MagicLeapMeshTracker.cpp
new file mode 100644
--- /dev/null
+++ b/Engine/Plugins/Lumin/MagicLeap/Source/MagicLeapHandMeshing/Private/MagicLeapMeshTracker.cpp
@@ -0,0 +1,361 @@
+// Copyright 1998-2018 Epic Games, Inc. All Rights Reserved.
+#include "MagicLeapMeshTracker.h"
+#include "HeadMountedDisplayFunctionLibrary.h"
+#include "MagicLeapHandMeshingTypes.h"
+#include "MagicLeapHandMeshingModule.h"
+#include "MagicLeapMath.h"
+
+FMagicLeapMeshTracker::FMagicLeapMeshTracker()
+: AppEventHandler() // HandMesh priv is autogranted and non-reality, so doesnt need runtime request.
+#if WITH_MLSDK
+, MeshTracker(ML_INVALID_HANDLE)
+, CurrentMeshRequest(ML_INVALID_HANDLE)
+#endif //WITH_MLSDK
+, bCreating(false)
+, MeshBrickIndex(0)
+, MRMesh(nullptr)
+, CurrentMeshDataCache(nullptr)
+{
+};
+
+void FMagicLeapMeshTracker::OnClear()
+{
+	MeshBrickIndex = 0;
+	MeshBrickCache.Empty();
+	PendingMeshBricks.Empty();
+}
+
+void FMagicLeapMeshTracker::FreeMeshDataCache(FMagicLeapMeshTracker::FMLCachedMeshData::SharedPtr& DataCache)
+{
+	FScopeLock ScopeLock(&FreeCachedMeshDatasMutex);
+	FreeCachedMeshDatas.Add(DataCache);
+}
+
+bool FMagicLeapMeshTracker::Create()
+{
+#if WITH_MLSDK
+	if (MLHandleIsValid(MeshTracker))
+	{
+		UE_LOG(LogMagicLeapHandMeshing, Warning, TEXT("FMagicLeapMeshTracker has already been created!"));
+		return false;
+	}
+
+	bCreating = true;
+#endif //WITH_MLSDK
+	return true;
+}
+
+bool FMagicLeapMeshTracker::Destroy()
+{
+	bCreating = false;
+#if WITH_MLSDK
+	if (!MLHandleIsValid(MeshTracker))
+	{
+		UE_LOG(LogMagicLeapHandMeshing, Warning, TEXT("FMagicLeapMeshTracker has already been destroyed!"));
+		return false;
+	}
+	else
+	{
+		DisconnectMRMesh(MRMesh);
+		MLResult Result = MLHandMeshingDestroyClient(&MeshTracker);
+		if (Result != MLResult_Ok)
+		{
+			UE_LOG(LogMagicLeapHandMeshing, Error, TEXT("MLHandMeshingDestroyClient failed with error '%s'"), UTF8_TO_TCHAR(MLGetResultString(Result)));
+			return false;
+		}
+
+		MeshTracker = ML_INVALID_HANDLE;
+	}
+#endif //WITH_MLSDK
+	return true;
+}
+
+bool FMagicLeapMeshTracker::ConnectMRMesh(UMRMeshComponent* InMRMeshPtr)
+{
+	if (!InMRMeshPtr)
+	{
+		UE_LOG(LogMagicLeapHandMeshing, Warning,
+			TEXT("MRMesh given is not valid. Ignoring this connect."));
+		return false;
+	}
+	else if (MRMesh)
+	{
+		UE_LOG(LogMagicLeapHandMeshing, Warning,
+			TEXT("FMagicLeapMeshTracker already has a MRMesh connected.  Ignoring this connect."));
+		return false;
+	}
+	else if (InMRMeshPtr->IsConnected())
+	{
+		UE_LOG(LogMagicLeapHandMeshing, Warning,
+			TEXT("MRMesh is already connected to a FMagicLeapMeshTracker. Ignoring this connect."));
+		return false;
+	}
+	else
+	{
+		InMRMeshPtr->SetConnected(true);
+		MRMesh = InMRMeshPtr;
+	}
+
+	return true;
+}
+
+bool FMagicLeapMeshTracker::DisconnectMRMesh(class UMRMeshComponent* InMRMeshPtr)
+{
+	if (!MRMesh)
+	{
+		UE_LOG(LogMagicLeapHandMeshing, Warning,
+			TEXT("FMagicLeapMeshTracker MRMesh is already disconnected. Ignoring this disconnect."));
+		return false;
+	}
+	else if (InMRMeshPtr != MRMesh)
+	{
+		UE_LOG(LogMagicLeapHandMeshing, Warning,
+			TEXT("FMagicLeapMeshTracker MRMesh given is not the MRMesh connected. Ignoring this disconnect."));
+		return false;
+	}
+	else
+	{
+		check(MRMesh->IsConnected());
+		MRMesh->SetConnected(false);
+	}
+	MRMesh = nullptr;
+	return true;
+}
+
+bool FMagicLeapMeshTracker::Update()
+{
+	if (bCreating)
+	{
+		if (!MRMesh)
+		{
+			UE_LOG(LogMagicLeapHandMeshing, Error, TEXT("FMagicLeapMeshTracker has no MRMesh!"));
+			return false;
+		}
+
+#if WITH_MLSDK
+		MLResult Result = MLHandMeshingCreateClient(&MeshTracker);
+		if (Result != MLResult_Ok && Result != MLResult_NotImplemented)
+		{
+			UE_LOG(LogMagicLeapHandMeshing, Error, TEXT("MLHandMeshingCreateClient failed with error '%s'"), UTF8_TO_TCHAR(MLGetResultString(Result)));
+			return false;
+		}
+#endif // WITH_MLSDK
+
+		MeshBrickCache.Empty();
+		MeshBrickIndex = 0;
+		bCreating = false;
+	}
+
+#if WITH_MLSDK
+	if (!MLHandleIsValid(MeshTracker))
+	{
+		return false;
+	}
+
+	if (!MLHandleIsValid(CurrentMeshRequest))
+	{
+		if (!RequestMesh())
+		{
+			return false;
+		}
+	}
+#endif // WITH_MLSDK
+
+	return GetMeshResult();
+}
+
+bool FMagicLeapMeshTracker::RequestMesh()
+{
+#if WITH_MLSDK
+	MLResult Result = MLHandMeshingRequestMesh(MeshTracker, &CurrentMeshRequest);
+	if (Result != MLResult_Ok)
+	{
+		UE_LOG(LogMagicLeapHandMeshing, Error, TEXT("MLHandMeshingRequestMesh failed with error %s"), UTF8_TO_TCHAR(MLGetResultString(Result)));
+		return false;
+	}
+
+	return true;
+#else
+	return false;
+#endif // WITH_MLSDK
+}
+
+bool FMagicLeapMeshTracker::GetMeshResult()
+{
+#if WITH_MLSDK
+	const float WorldToMetersScale = IMagicLeapPlugin::Get().GetWorldToMetersScale();
+
+	// Get mesh result
+	if (MLHandleIsValid(CurrentMeshRequest))
+	{
+		MLHandMesh Mesh = {};
+		MLHandMeshInit(&Mesh);
+		auto Result = MLHandMeshingGetResult(MeshTracker, CurrentMeshRequest, &Mesh);
+
+		if (MLResult_Ok != Result)
+		{
+			// Just silently wait for pending result
+			if (MLResult_Pending != Result)
+			{
+				UE_LOG(LogMagicLeapHandMeshing, Error, TEXT("MLMeshingGetMeshResult failed: %s."), UTF8_TO_TCHAR(MLGetResultString(Result)));
+				// Mesh request failed, lets queue another one.
+				CurrentMeshRequest = ML_INVALID_HANDLE;
+				return true;
+			}
+
+			if (CurrentMeshDataCache.Get() != nullptr)
+			{
+				// Get/create brick ID for this mesh GUID
+				auto MeshBrickInfo = GetBrickInfo(CurrentMeshDataCache->BlockID, true);
+				// Create/update brick
+				static_cast<IMRMesh*>(MRMesh)->SendBrickData(
+					IMRMesh::FSendBrickDataArgs
+					{
+						TSharedPtr<IMRMesh::FBrickDataReceipt, ESPMode::ThreadSafe>
+						(new FMeshTrackerBrickDataReceipt(CurrentMeshDataCache)),
+						*MeshBrickInfo,
+						CurrentMeshDataCache->WorldVertices,
+						CurrentMeshDataCache->UV0,
+						CurrentMeshDataCache->Tangents,
+						CurrentMeshDataCache->VertexColors,
+						CurrentMeshDataCache->Triangles
+					}
+				);
+			}
+			// Mesh request pending...
+			return false;
+		}
+		else
+		{
+			FVector VertexOffset = UHeadMountedDisplayFunctionLibrary::GetTrackingToWorldTransform(MRMesh).Inverse().GetLocation();
+			for (uint32_t MeshIndex = 0; MeshIndex < Mesh.data_count; ++MeshIndex)
+			{
+				const auto &MeshData = Mesh.data[MeshIndex];
+				auto BlockID = FGuid();// MeshData.id.data[0], MeshData.id.data[0] >> 32, MeshData.id.data[1], MeshData.id.data[1] >> 32);
+				// Acquire mesh data cache and mark its brick ID
+				CurrentMeshDataCache = AquireMeshDataCache();
+				CurrentMeshDataCache->BlockID = BlockID;
+
+				// Pull vertices
+				CurrentMeshDataCache->OffsetVertices.Reserve(MeshData.vertex_count);
+				CurrentMeshDataCache->WorldVertices.Reserve(MeshData.vertex_count);
+				for (uint32_t v = 0; v < MeshData.vertex_count; ++v)
+				{
+					CurrentMeshDataCache->OffsetVertices.Add(MagicLeap::ToFVector(MeshData.vertex[v], WorldToMetersScale) - VertexOffset);
+					CurrentMeshDataCache->WorldVertices.Add(MagicLeap::ToFVector(MeshData.vertex[v], WorldToMetersScale));
+				}
+
+				// Pull indices
+				CurrentMeshDataCache->Triangles.Reserve(MeshData.index_count);
+				for (uint16_t i = 0; i < MeshData.index_count; ++i)
+				{
+					CurrentMeshDataCache->Triangles.Add(static_cast<uint32>(MeshData.index[i]));
+				}
+
+				// Pull normals
+				CurrentMeshDataCache->Normals.Reserve(MeshData.vertex_count);
+				for (int32 iNormal = 0; iNormal < CurrentMeshDataCache->Triangles.Num()-2; ++iNormal)
+				{
+					FVector A = CurrentMeshDataCache->WorldVertices[CurrentMeshDataCache->Triangles[iNormal]];
+					FVector B = CurrentMeshDataCache->WorldVertices[CurrentMeshDataCache->Triangles[iNormal+1]];
+					FVector C = CurrentMeshDataCache->WorldVertices[CurrentMeshDataCache->Triangles[iNormal+2]];
+					FVector AToB = B - A;
+					AToB.Normalize();
+					FVector AToC = C - A;
+					AToC.Normalize();
+					FVector Normal = FVector::CrossProduct(AToB, AToC);
+					CurrentMeshDataCache->Normals.Add(Normal);
+				}
+				
+				CurrentMeshDataCache->Tangents.Reserve(MeshData.vertex_count * 2);
+				for (uint32_t t = 0; t < MeshData.vertex_count; ++t)
+				{
+					const FVector& Norm = CurrentMeshDataCache->Normals[t];
+
+					// Calculate tangent
+					auto Perp = Norm.X < Norm.Z ?
+						FVector(1.0f, 0.0f, 0.0f) : FVector(0.0f, 1.0f, 0.0f);
+					auto Tang = FVector::CrossProduct(Norm, Perp);
+					Tang.Normalize();
+
+					CurrentMeshDataCache->Tangents.Add(Tang);
+					CurrentMeshDataCache->Tangents.Add(Norm);
+				}
+
+				// To work in all rendering paths we always set a vertex color
+				if (CurrentMeshDataCache->VertexColors.Num() == 0)
+				{
+					for (uint32 v = 0; v < MeshData.vertex_count; ++v)
+					{
+						CurrentMeshDataCache->VertexColors.Add(FColor::White);
+					}
+				}
+
+				// Write UVs
+				CurrentMeshDataCache->UV0.Reserve(MeshData.vertex_count);
+				for (uint32 v = 0; v < MeshData.vertex_count; ++v)
+				{
+					const float FakeCoord = static_cast<float>(v) / static_cast<float>(MeshData.vertex_count);
+					CurrentMeshDataCache->UV0.Add(FVector2D(FakeCoord, FakeCoord));
+				}
+
+#ifdef DEBUG_MESH_BRICK_EVENTS
+				UE_LOG(LogMagicLeapHandMeshing,
+					Log,
+					TEXT("UMeshTrackerComponent: ADDING/UPDATING brick %s"),
+					TCHAR_TO_UTF8(*(BlockID.ToString())));
+#endif //DEBUG_MESH_BRICK_EVENTS
+
+				// Get/create brick ID for this mesh GUID
+				auto MeshBrickInfo = GetBrickInfo(BlockID, true);
+				// Create/update brick
+				static_cast<IMRMesh*>(MRMesh)->SendBrickData(
+					IMRMesh::FSendBrickDataArgs
+					{
+						TSharedPtr<IMRMesh::FBrickDataReceipt, ESPMode::ThreadSafe>
+						(new FMeshTrackerBrickDataReceipt(CurrentMeshDataCache)),
+						*MeshBrickInfo,
+						CurrentMeshDataCache->WorldVertices,
+						CurrentMeshDataCache->UV0,
+						CurrentMeshDataCache->Tangents,
+						CurrentMeshDataCache->VertexColors,
+						CurrentMeshDataCache->Triangles
+					}
+				);
+			}
+
+			// All meshes pulled and/or updated; free the ML resource
+			MLHandMeshingFreeResource(MeshTracker, &CurrentMeshRequest);
+			CurrentMeshRequest = ML_INVALID_HANDLE;
+			CurrentMeshDataCache = nullptr;
+			return true;
+		}
+	}
+#endif // WITH_MLSDK
+	return true;
+}
+
+uint64* FMagicLeapMeshTracker::GetBrickInfo(const FGuid& meshGuid, bool addIfNotFound)
+{
+	auto meshBrickInfo = MeshBrickCache.Find(meshGuid);
+	if (meshBrickInfo == nullptr && addIfNotFound)
+	{
+		MeshBrickCache.Add(meshGuid, MeshBrickIndex++);
+		meshBrickInfo = MeshBrickCache.Find(meshGuid);
+	}
+	return meshBrickInfo;
+}
+
+bool FMagicLeapMeshTracker::HasMRMesh() const
+{
+	return MRMesh != nullptr;
+}
+
+bool FMagicLeapMeshTracker::HasClient() const
+{
+#if WITH_MLSDK
+	return MLHandleIsValid(MeshTracker);
+#else
+	return false;
+#endif // WITH_MLSDK
+}
diff --git a/Engine/Plugins/Lumin/MagicLeap/Source/MagicLeapHandMeshing/Private/MagicLeapMeshTracker.h b/Engine/Plugins/Lumin/MagicLeap/Source/MagicLeapHandMeshing/Private/MagicLeapMeshTracker.h
new file mode 100644
--- /dev/null
+++ b/Engine/Plugins/Lumin/MagicLeap/Source/MagicLeapHandMeshing/Private/MagicLeapMeshTracker.h
@@ -0,0 +1,131 @@
+// Copyright 1998-2018 Epic Games, Inc. All Rights Reserved.
+
+#pragma once
+
+#include "CoreMinimal.h"
+#include "PackedNormal.h"
+#include "Lumin/CAPIShims/LuminAPIHandMeshing.h"
+#include "MRMeshComponent.h"
+#include "AppEventHandler.h"
+
+class FMagicLeapMeshTracker
+{
+public:
+	FMagicLeapMeshTracker();
+	virtual ~FMagicLeapMeshTracker() {}
+
+	void OnClear();
+
+public:
+	// Keep a copy of the mesh data here.  MRMeshComponent will use it from the game and render thread.
+	struct FMLCachedMeshData
+	{
+		typedef TSharedPtr<FMLCachedMeshData, ESPMode::ThreadSafe> SharedPtr;
+
+		FMagicLeapMeshTracker* Owner = nullptr;
+
+		FGuid BlockID;
+		TArray<FVector> OffsetVertices;
+		TArray<FVector> WorldVertices;
+		TArray<uint32> Triangles;
+		TArray<FVector> Normals;
+		TArray<FVector2D> UV0;
+		TArray<FColor> VertexColors;
+		TArray<FPackedNormal> Tangents;
+		TArray<float> Confidence;
+
+		void Recycle(SharedPtr& MeshData)
+		{
+			check(Owner);
+			FMagicLeapMeshTracker* TempOwner = Owner;
+			Owner = nullptr;
+
+			BlockID.Invalidate();
+			OffsetVertices.Reset();
+			WorldVertices.Reset();
+			Triangles.Reset();
+			Normals.Reset();
+			UV0.Reset();
+			VertexColors.Reset();
+			Tangents.Reset();
+			Confidence.Reset();
+
+			TempOwner->FreeMeshDataCache(MeshData);
+		}
+
+		void Init(FMagicLeapMeshTracker* InOwner)
+		{
+			check(!Owner);
+			Owner = InOwner;
+		}
+	};
+
+	// Pending MR Mesh bricks, generated from mesh update messages, which will be amortized
+	// across frames to avoid spiking
+	TMap<FGuid, FMLCachedMeshData::SharedPtr> PendingMeshBricks;
+
+	FDelegateHandle OnClearDelegateHandle;
+
+	// This receipt will be kept in the FSendBrickDataArgs to ensure the cached data outlives MRMeshComponent use of it.
+	class FMeshTrackerBrickDataReceipt : public IMRMesh::FBrickDataReceipt
+	{
+	public:
+		FMeshTrackerBrickDataReceipt(FMLCachedMeshData::SharedPtr& MeshData) : CachedMeshData(MeshData)
+		{
+		}
+		~FMeshTrackerBrickDataReceipt() override
+		{
+			CachedMeshData->Recycle(CachedMeshData);
+		}
+	private:
+		FMLCachedMeshData::SharedPtr CachedMeshData;
+	};
+
+	FMLCachedMeshData::SharedPtr AquireMeshDataCache()
+	{
+		if (FreeCachedMeshDatas.Num() > 0)
+		{
+			FScopeLock ScopeLock(&FreeCachedMeshDatasMutex);
+			FMLCachedMeshData::SharedPtr CachedMeshData(FreeCachedMeshDatas.Pop(false));
+			CachedMeshData->Init(this);
+			return CachedMeshData;
+		}
+		else
+		{
+			FMLCachedMeshData::SharedPtr CachedMeshData(new FMLCachedMeshData());
+			CachedMeshData->Init(this);
+			CachedMeshDatas.Add(CachedMeshData);
+			return CachedMeshData;
+		}
+	}
+
+	void FreeMeshDataCache(FMLCachedMeshData::SharedPtr& DataCache);
+	bool Create();
+	bool Destroy();
+	bool ConnectMRMesh(UMRMeshComponent* InMRMeshPtr);
+	bool DisconnectMRMesh(class UMRMeshComponent* InMRMeshPtr);
+	bool Update();
+	bool RequestMesh();
+	bool GetMeshResult();
+	uint64 *GetBrickInfo(const FGuid& meshGuid, bool addIfNotFound);
+	bool HasMRMesh() const;
+	bool HasClient() const;
+
+private:
+	MagicLeap::IAppEventHandler AppEventHandler;
+#if WITH_MLSDK
+	// Handle to ML mesh tracker
+	MLHandle MeshTracker;
+	MLHandle CurrentMeshRequest;
+#endif //WITH_MLSDK
+	bool bCreating;
+	// Maps system 128-bit mesh GUIDs to 64-bit brick indices
+	TMap<FGuid, uint64> MeshBrickCache;
+	uint64 MeshBrickIndex;
+	class UMRMeshComponent* MRMesh;
+	// A free list to recycle the CachedMeshData instances.
+	FMLCachedMeshData::SharedPtr CurrentMeshDataCache;
+	TArray<FMLCachedMeshData::SharedPtr> CachedMeshDatas;
+	TArray<FMLCachedMeshData::SharedPtr> FreeCachedMeshDatas;
+	FCriticalSection FreeCachedMeshDatasMutex; //The free list may be pushed/popped from multiple threads.
+};
diff --git a/Engine/Plugins/Lumin/MagicLeap/Source/MagicLeapHandMeshing/Public/IMagicLeapHandMeshingModule.h b/Engine/Plugins/Lumin/MagicLeap/Source/MagicLeapHandMeshing/Public/IMagicLeapHandMeshingModule.h
new file mode 100644
--- /dev/null
+++ b/Engine/Plugins/Lumin/MagicLeap/Source/MagicLeapHandMeshing/Public/IMagicLeapHandMeshingModule.h
@@ -0,0 +1,38 @@
+// Copyright 1998-2018 Epic Games, Inc. All Rights Reserved.
+
+#pragma once
+
+#include "Modules/ModuleManager.h"
+#include "Containers/Ticker.h"
+
+/**
+ * The public interface to this module.  In most cases, this interface is only public to sibling modules
+ * within this plugin.
+ */
+class IMagicLeapHandMeshingModule : public IModuleInterface
+{
+public:
+
+	/**
+	 * Singleton-like access to this module's interface.  This is just for convenience!
+	 * Beware of calling this during the shutdown phase, though.  Your module might have been unloaded already.
+	 *
+	 * @return Returns singleton instance, loading the module on demand if needed
+	 */
+	static inline IMagicLeapHandMeshingModule& Get()
+	{
+		return FModuleManager::LoadModuleChecked<IMagicLeapHandMeshingModule>("MagicLeapHandMeshing");
+	}
+
+	/**
+	 * Checks to see if this module is loaded and ready.  It is only valid to call Get() if IsAvailable() returns true.
+	 *
+	 * @return True if the module is loaded and ready to use
+	 */
+	static inline bool IsAvailable()
+	{
+		return FModuleManager::Get().IsModuleLoaded("MagicLeapHandMeshing");
+	}
+
+	virtual bool Tick(float DeltaTime) = 0;
+};
diff --git a/Engine/Plugins/Lumin/MagicLeap/Source/MagicLeapHandMeshing/Public/MagicLeapHandMeshingComponent.h b/Engine/Plugins/Lumin/MagicLeap/Source/MagicLeapHandMeshing/Public/MagicLeapHandMeshingComponent.h
new file mode 100644
--- /dev/null
+++ b/Engine/Plugins/Lumin/MagicLeap/Source/MagicLeapHandMeshing/Public/MagicLeapHandMeshingComponent.h
@@ -0,0 +1,39 @@
+// Copyright 1998-2018 Epic Games, Inc. All Rights Reserved.
+
+#pragma once
+
+#include "Engine/Engine.h"
+#include "MRMeshComponent.h"
+#include "MagicLeapHandMeshingTypes.h"
+#include "MagicLeapHandMeshingComponent.generated.h"
+
+/**
+	Component that provides access to the HandMeshing API functionality.
+*/
+UCLASS(ClassGroup = MagicLeap, BlueprintType, Blueprintable, EditInlineNew, meta = (BlueprintSpawnableComponent))
+class MAGICLEAPHANDMESHING_API UMagicLeapHandMeshingComponent : public UActorComponent
+{
+	GENERATED_BODY()
+
+public:
+	/**
+		Destroys the hand meshing client.
+	*/
+	void EndPlay(const EEndPlayReason::Type EndPlayReason);
+
+	/**
+		Connects the MRMesh component.
+		@param InMRMeshPtr The MRMeshComponent to be used as the hand meshing source.
+		@return True if successful, false otherwise.
+	*/
+	UFUNCTION(BlueprintCallable, Category = "HandMeshing Function Library | MagicLeap")
+	bool ConnectMRMesh(UMRMeshComponent* InMRMeshPtr);
+
+	/**
+		Disconnects the MRMesh component.
+		@param InMRMeshPtr The MRMeshComponent to be removed as the hand meshing source.
+		@return True if successful, false otherwise.
+	*/
+	UFUNCTION(BlueprintCallable, Category = "HandMeshing Function Library | MagicLeap")
+	bool DisconnectMRMesh(class UMRMeshComponent* InMRMeshPtr);
+};
diff --git a/Engine/Plugins/Lumin/MagicLeap/Source/MagicLeapHandMeshing/Public/MagicLeapHandMeshingFunctionLibrary.h b/Engine/Plugins/Lumin/MagicLeap/Source/MagicLeapHandMeshing/Public/MagicLeapHandMeshingFunctionLibrary.h
new file mode 100644
--- /dev/null
+++ b/Engine/Plugins/Lumin/MagicLeap/Source/MagicLeapHandMeshing/Public/MagicLeapHandMeshingFunctionLibrary.h
@@ -0,0 +1,45 @@
+// Copyright 1998-2018 Epic Games, Inc. All Rights Reserved.
+#pragma once
+
+#include "Engine/Engine.h"
+#include "Kismet/BlueprintFunctionLibrary.h"
+#include "MagicLeapHandMeshingTypes.h"
+#include "MagicLeapHandMeshingFunctionLibrary.generated.h"
+
+UCLASS(ClassGroup = MagicLeap)
+class MAGICLEAPHANDMESHING_API UMagicLeapHandMeshingFunctionLibrary : public UBlueprintFunctionLibrary
+{
+	GENERATED_BODY()
+
+public:
+	/**
+		Create the hand meshing client Note that this will be the only function in the HandMeshing API that will raise a
+		PrivilegeDenied error.  Trying to call the other functions with an invalid MLHandle will result in an InvalidParam error.
+		@return True if successful, false otherwise.
+	*/
+	UFUNCTION(BlueprintCallable, Category = "HandMeshing Function Library | MagicLeap")
+	static bool CreateClient();
+
+	/** 
+		Destroys the hand meshing client.
+		@return True if successful, false otherwise.
+	*/
+	UFUNCTION(BlueprintCallable, Category = "HandMeshing Function Library | MagicLeap")
+	static bool DestroyClient();
+
+	/** 
+		Connects the MRMesh component.
+		@param InMRMeshPtr The MRMeshComponent to be used as the hand meshing source.
+		@return True if successful, false otherwise.
+	*/
+	UFUNCTION(BlueprintCallable, Category = "HandMeshing Function Library | MagicLeap")
+	static bool ConnectMRMesh(UMRMeshComponent* InMRMeshPtr);
+
+	/** 
+		Disconnects the MRMesh component.
+		@param InMRMeshPtr The MRMeshComponent to be removed as the hand meshing source.
+		@return True if successful, false otherwise.
+	*/
+	UFUNCTION(BlueprintCallable, Category = "HandMeshing Function Library | MagicLeap")
+	static bool DisconnectMRMesh(class UMRMeshComponent* InMRMeshPtr);
+};
diff --git a/Engine/Plugins/Lumin/MagicLeap/Source/MagicLeapHandMeshing/Public/MagicLeapHandMeshingTypes.h b/Engine/Plugins/Lumin/MagicLeap/Source/MagicLeapHandMeshing/Public/MagicLeapHandMeshingTypes.h
new file mode 100644
--- /dev/null
+++ b/Engine/Plugins/Lumin/MagicLeap/Source/MagicLeapHandMeshing/Public/MagicLeapHandMeshingTypes.h
@@ -0,0 +1,47 @@
+// Copyright 1998-2018 Epic Games, Inc. All Rights Reserved.
+#pragma once
+
+#include "Engine/Engine.h"
+#include "MagicLeapHandMeshingTypes.generated.h"
+
+/** Stores a hand mesh's vertices and indices. */
+USTRUCT(BlueprintType)
+struct FMagicLeapHandMeshBlock
+{
+	GENERATED_USTRUCT_BODY()
+
+	/** The number of indices in index buffer. */
+	UPROPERTY(EditAnywhere, BlueprintReadWrite, Category = "HandMeshing|MagicLeap")
+	int32 IndexCount;
+
+	/** The number of vertices in vertex buffer. */
+	UPROPERTY(EditAnywhere, BlueprintReadWrite, Category = "HandMeshing|MagicLeap")
+	int32 VertexCount;
+
+	/** Pointer to the vertex buffer. */
+	UPROPERTY(EditAnywhere, BlueprintReadWrite, Category = "HandMeshing|MagicLeap")
+	TArray<FVector> Vertex;
+
+	/** Pointer to the index buffer. */
+	UPROPERTY(EditAnywhere, BlueprintReadWrite, Category = "HandMeshing|MagicLeap")
+	TArray<int32> Index;
+};
+
+/** Stores MLHandMeshBlock data. */
+USTRUCT(BlueprintType)
+struct FMagicLeapHandMesh
+{
+	GENERATED_USTRUCT_BODY()
+
+	/** Version of this structure. */
+	UPROPERTY()
+	int32 Version;
+
+	/** Number of meshes available in data */
+	UPROPERTY(EditAnywhere, BlueprintReadWrite, Category = "HandMeshing|MagicLeap")
+	int32 DataCount;
+
+	/** The mesh data */
+	UPROPERTY(EditAnywhere, BlueprintReadWrite, Category = "HandMeshing|MagicLeap")
+	TArray<FMagicLeapHandMeshBlock> Data;
+};
diff --git a/Engine/Plugins/Lumin/MagicLeap/Source/MagicLeapHandTracking/MagicLeapHandTracking.Build.cs b/Engine/Plugins/Lumin/MagicLeap/Source/MagicLeapHandTracking/MagicLeapHandTracking.Build.cs
--- a/Engine/Plugins/Lumin/MagicLeap/Source/MagicLeapHandTracking/MagicLeapHandTracking.Build.cs
+++ b/Engine/Plugins/Lumin/MagicLeap/Source/MagicLeapHandTracking/MagicLeapHandTracking.Build.cs
@@ -32,9 +32,6 @@ namespace UnrealBuildTool.Rules
 					"LiveLinkInterface",
 				}
 			);
-
-			// This is not ideal but needs to be done in order to expose the private MagicLeapHMD header to this module.
-			PrivateIncludePaths.Add(Path.Combine(new string[] { ModuleDirectory, "..", "MagicLeap", "Private" }));
 		}
 	}
 }
