Copying //UE4/Partner-MagicLeap-Staging to Dev-VR (//UE4/Dev-VR)

diff --git a/Engine/Plugins/Lumin/MagicLeap/Source/MagicLeapHandTracking/Private/MagicLeapHandTracking.cpp b/Engine/Plugins/Lumin/MagicLeap/Source/MagicLeapHandTracking/Private/MagicLeapHandTracking.cpp
--- a/Engine/Plugins/Lumin/MagicLeap/Source/MagicLeapHandTracking/Private/MagicLeapHandTracking.cpp
+++ b/Engine/Plugins/Lumin/MagicLeap/Source/MagicLeapHandTracking/Private/MagicLeapHandTracking.cpp
@@ -2,7 +2,6 @@
 
 #include "MagicLeapHandTracking.h"
 #include "IMagicLeapPlugin.h"
-#include "MagicLeapHMD.h"
 #include "MagicLeapMath.h"
 #include "Framework/Application/SlateApplication.h"
 #include "CoreMinimal.h"
@@ -95,10 +94,10 @@ IMPLEMENT_MODULE(FMagicLeapHandTrackingPlugin, MagicLeapHandTracking);
 
 // Partial sanity check of UE4 enums vs ML enums.  If these fail there has probably been an MLSDK update that requires an UE4 update.
 #if WITH_MLSDK
-static_assert(EHandTrackingKeypointCount == static_cast<int32>(MLHandTrackingStaticData_MaxKeyPoints), "EHandTrackingGesture does not match api enum.");
-static_assert(static_cast<int32>(EHandTrackingGesture::NoHand) == static_cast<int32>(MLHandTrackingKeyPose::MLHandTrackingKeyPose_NoHand), "EHandTrackingGesture does not match api enum.");
-static_assert(static_cast<int32>(EHandTrackingKeypointFilterLevel::PredictiveSmoothing) == static_cast<int32>(MLKeypointFilterLevel::MLKeypointFilterLevel_2), "EHandTrackingKeypointFilterLevel does not match api enum.");
-static_assert(static_cast<int32>(EHandTrackingGestureFilterLevel::MoreRobustnessToFlicker) == static_cast<int32>(MLPoseFilterLevel::MLPoseFilterLevel_2), "EHandTrackingGestureFilterLevel does not match api enum.");
+static_assert(EMagicLeapHandTrackingKeypointCount == static_cast<int32>(MLHandTrackingStaticData_MaxKeyPoints), "EMagicLeapHandTrackingGesture does not match api enum.");
+static_assert(static_cast<int32>(EMagicLeapHandTrackingGesture::NoHand) == static_cast<int32>(MLHandTrackingKeyPose::MLHandTrackingKeyPose_NoHand), "EMagicLeapHandTrackingGesture does not match api enum.");
+static_assert(static_cast<int32>(EMagicLeapHandTrackingKeypointFilterLevel::PredictiveSmoothing) == static_cast<int32>(MLKeypointFilterLevel::MLKeypointFilterLevel_2), "EMagicLeapHandTrackingKeypointFilterLevel does not match api enum.");
+static_assert(static_cast<int32>(EMagicLeapHandTrackingGestureFilterLevel::MoreRobustnessToFlicker) == static_cast<int32>(MLPoseFilterLevel::MLPoseFilterLevel_2), "EMagicLeapHandTrackingGestureFilterLevel does not match api enum.");
 #endif
 
 // Left Gestures
@@ -117,6 +116,9 @@ const FName FMagicLeapGestureKeyNames::Left_Thumb_Name("MagicLeap_Left_Thumb");
 const FKey FMagicLeapHandTracking::FStaticHandTracking::Left_L("MagicLeap_Left_L");
 const FName FMagicLeapGestureKeyNames::Left_L_Name("MagicLeap_Left_L");
 
+const FKey FMagicLeapHandTracking::FStaticHandTracking::Left_OpenHand("MagicLeap_Left_OpenHand");
+const FName FMagicLeapGestureKeyNames::Left_OpenHand_Name("MagicLeap_Left_OpenHand");
+// DEPRECATED
 const FKey FMagicLeapHandTracking::FStaticHandTracking::Left_OpenHandBack("MagicLeap_Left_OpenHandBack");
 const FName FMagicLeapGestureKeyNames::Left_OpenHandBack_Name("MagicLeap_Left_OpenHandBack");
 
@@ -148,6 +150,9 @@ const FName FMagicLeapGestureKeyNames::Right_Thumb_Name("MagicLeap_Right_Thumb")
 const FKey FMagicLeapHandTracking::FStaticHandTracking::Right_L("MagicLeap_Right_L");
 const FName FMagicLeapGestureKeyNames::Right_L_Name("MagicLeap_Right_L");
 
+const FKey FMagicLeapHandTracking::FStaticHandTracking::Right_OpenHand("MagicLeap_Right_OpenHand");
+const FName FMagicLeapGestureKeyNames::Right_OpenHand_Name("MagicLeap_Right_OpenHand");
+// DEPRECATED
 const FKey FMagicLeapHandTracking::FStaticHandTracking::Right_OpenHandBack("MagicLeap_Right_OpenHandBack");
 const FName FMagicLeapGestureKeyNames::Right_OpenHandBack_Name("MagicLeap_Right_OpenHandBack");
 
@@ -272,60 +277,60 @@ FName FMagicLeapHandTracking::LiveLinkRightHandTrackingSubjectName(TEXT("MagicLe
 
 
 #if WITH_MLSDK
-EHandTrackingKeypointFilterLevel MLToUnrealKeypointsFilterLevel(MLKeypointFilterLevel FilterLevel)
+EMagicLeapHandTrackingKeypointFilterLevel MLToUnrealKeypointsFilterLevel(MLKeypointFilterLevel FilterLevel)
 {
 	switch (FilterLevel)
 	{
 	case MLKeypointFilterLevel_0:
-		return EHandTrackingKeypointFilterLevel::NoFilter;
+		return EMagicLeapHandTrackingKeypointFilterLevel::NoFilter;
 	case MLKeypointFilterLevel_1:
-		return EHandTrackingKeypointFilterLevel::SimpleSmoothing;
+		return EMagicLeapHandTrackingKeypointFilterLevel::SimpleSmoothing;
 	case MLKeypointFilterLevel_2:
-		return EHandTrackingKeypointFilterLevel::PredictiveSmoothing;
+		return EMagicLeapHandTrackingKeypointFilterLevel::PredictiveSmoothing;
 	}
 
-	return EHandTrackingKeypointFilterLevel::NoFilter;
+	return EMagicLeapHandTrackingKeypointFilterLevel::NoFilter;
 }
 
-EHandTrackingGestureFilterLevel MLToUnrealGestureFilterLevel(MLPoseFilterLevel FilterLevel)
+EMagicLeapHandTrackingGestureFilterLevel MLToUnrealGestureFilterLevel(MLPoseFilterLevel FilterLevel)
 {
 	switch (FilterLevel)
 	{
 	case MLPoseFilterLevel_0:
-		return EHandTrackingGestureFilterLevel::NoFilter;
+		return EMagicLeapHandTrackingGestureFilterLevel::NoFilter;
 	case MLPoseFilterLevel_1:
-		return EHandTrackingGestureFilterLevel::SlightRobustnessToFlicker;
+		return EMagicLeapHandTrackingGestureFilterLevel::SlightRobustnessToFlicker;
 	case MLPoseFilterLevel_2:
-		return EHandTrackingGestureFilterLevel::MoreRobustnessToFlicker;
+		return EMagicLeapHandTrackingGestureFilterLevel::MoreRobustnessToFlicker;
 	}
 
-	return EHandTrackingGestureFilterLevel::NoFilter;
+	return EMagicLeapHandTrackingGestureFilterLevel::NoFilter;
 }
 
-MLKeypointFilterLevel UnrealToMLKeypointsFilterLevel(EHandTrackingKeypointFilterLevel FilterLevel)
+MLKeypointFilterLevel UnrealToMLKeypointsFilterLevel(EMagicLeapHandTrackingKeypointFilterLevel FilterLevel)
 {
 	switch (FilterLevel)
 	{
-	case EHandTrackingKeypointFilterLevel::NoFilter:
+	case EMagicLeapHandTrackingKeypointFilterLevel::NoFilter:
 		return MLKeypointFilterLevel_0;
-	case EHandTrackingKeypointFilterLevel::SimpleSmoothing:
+	case EMagicLeapHandTrackingKeypointFilterLevel::SimpleSmoothing:
 		return MLKeypointFilterLevel_1;
-	case EHandTrackingKeypointFilterLevel::PredictiveSmoothing:
+	case EMagicLeapHandTrackingKeypointFilterLevel::PredictiveSmoothing:
 		return MLKeypointFilterLevel_2;
 	}
 
 	return MLKeypointFilterLevel_0;
 }
 
-MLPoseFilterLevel UnrealToMLGestureFilterLevel(EHandTrackingGestureFilterLevel FilterLevel)
+MLPoseFilterLevel UnrealToMLGestureFilterLevel(EMagicLeapHandTrackingGestureFilterLevel FilterLevel)
 {
 	switch (FilterLevel)
 	{
-	case EHandTrackingGestureFilterLevel::NoFilter:
+	case EMagicLeapHandTrackingGestureFilterLevel::NoFilter:
 		return MLPoseFilterLevel_0;
-	case EHandTrackingGestureFilterLevel::SlightRobustnessToFlicker:
+	case EMagicLeapHandTrackingGestureFilterLevel::SlightRobustnessToFlicker:
 		return MLPoseFilterLevel_1;
-	case EHandTrackingGestureFilterLevel::MoreRobustnessToFlicker:
+	case EMagicLeapHandTrackingGestureFilterLevel::MoreRobustnessToFlicker:
 		return MLPoseFilterLevel_2;
 	}
 
@@ -363,7 +368,7 @@ FMagicLeapHandTracking::FMagicLeapHandTracking(const TSharedRef<FGenericApplicat
 
 	// We're implicitly requiring that the MagicLeapPlugin has been loaded and
 	// initialized at this point.
-	IMagicLeapPlugin::Get().RegisterMagicLeapInputDevice(this);
+	IMagicLeapPlugin::Get().RegisterMagicLeapTrackerEntity(this);
 }
 
 FMagicLeapHandTracking::~FMagicLeapHandTracking()
@@ -372,10 +377,10 @@ FMagicLeapHandTracking::~FMagicLeapHandTracking()
 	// but it isn't an assumption that we should make.
 	if (IMagicLeapPlugin::IsAvailable())
 	{
-		IMagicLeapPlugin::Get().UnregisterMagicLeapInputDevice(this);
+		IMagicLeapPlugin::Get().UnregisterMagicLeapTrackerEntity(this);
 	}
 
-	Disable();
+	DestroyEntityTracker();
 
 	// Unregister "MotionController" modular feature manually
 	IModularFeatures::Get().UnregisterModularFeature(GetModularFeatureName(), this);
@@ -383,7 +388,7 @@ FMagicLeapHandTracking::~FMagicLeapHandTracking()
 
 void FMagicLeapHandTracking::BuildKeypointMaps()
 {
-	SourceToTransformMap.Reserve(EHandTrackingKeypointCount * 2);
+	SourceToTransformMap.Reserve(EMagicLeapHandTrackingKeypointCount * 2);
 
 	SourceToTransformMap.Add(LeftHandCenter_Name, &LeftHand.HandCenter.Transform);
 
@@ -450,7 +455,7 @@ void FMagicLeapHandTracking::BuildKeypointMaps()
 
 
 
-	SourceToHandMap.Reserve(EHandTrackingKeypointCount * 2);
+	SourceToHandMap.Reserve(EMagicLeapHandTrackingKeypointCount * 2);
 
 	SourceToHandMap.Add(LeftHandCenter_Name, &LeftHand);
 
@@ -518,48 +523,48 @@ void FMagicLeapHandTracking::BuildKeypointMaps()
 
 FMagicLeapHandTracking::FHandState::FHandState()
 {
-	EnumToTransformMap[static_cast<int32>(EHandTrackingKeypoint::Thumb_Tip)] = &Thumb.Tip;
-	EnumToTransformMap[static_cast<int32>(EHandTrackingKeypoint::Thumb_IP)] = &Thumb.DIP;
-	EnumToTransformMap[static_cast<int32>(EHandTrackingKeypoint::Thumb_MCP)] = &Thumb.PIP;
-	EnumToTransformMap[static_cast<int32>(EHandTrackingKeypoint::Thumb_CMC)] = &Thumb.MCP;
-
-	EnumToTransformMap[static_cast<int32>(EHandTrackingKeypoint::Index_Tip)] = &IndexFinger.Tip;
-	EnumToTransformMap[static_cast<int32>(EHandTrackingKeypoint::Index_DIP)] = &IndexFinger.DIP;
-	EnumToTransformMap[static_cast<int32>(EHandTrackingKeypoint::Index_PIP)] = &IndexFinger.PIP;
-	EnumToTransformMap[static_cast<int32>(EHandTrackingKeypoint::Index_MCP)] = &IndexFinger.MCP;
-
-	EnumToTransformMap[static_cast<int32>(EHandTrackingKeypoint::Middle_Tip)] = &MiddleFinger.Tip;
-	EnumToTransformMap[static_cast<int32>(EHandTrackingKeypoint::Middle_DIP)] = &MiddleFinger.DIP;
-	EnumToTransformMap[static_cast<int32>(EHandTrackingKeypoint::Middle_PIP)] = &MiddleFinger.PIP;
-	EnumToTransformMap[static_cast<int32>(EHandTrackingKeypoint::Middle_MCP)] = &MiddleFinger.MCP;
-
-	EnumToTransformMap[static_cast<int32>(EHandTrackingKeypoint::Ring_Tip)] = &RingFinger.Tip;
-	EnumToTransformMap[static_cast<int32>(EHandTrackingKeypoint::Ring_DIP)] = &RingFinger.DIP;
-	EnumToTransformMap[static_cast<int32>(EHandTrackingKeypoint::Ring_PIP)] = &RingFinger.PIP;
-	EnumToTransformMap[static_cast<int32>(EHandTrackingKeypoint::Ring_MCP)] = &RingFinger.MCP;
-
-	EnumToTransformMap[static_cast<int32>(EHandTrackingKeypoint::Pinky_Tip)] = &PinkyFinger.Tip;
-	EnumToTransformMap[static_cast<int32>(EHandTrackingKeypoint::Pinky_DIP)] = &PinkyFinger.DIP;
-	EnumToTransformMap[static_cast<int32>(EHandTrackingKeypoint::Pinky_PIP)] = &PinkyFinger.PIP;
-	EnumToTransformMap[static_cast<int32>(EHandTrackingKeypoint::Pinky_MCP)] = &PinkyFinger.MCP;
-
-	EnumToTransformMap[static_cast<int32>(EHandTrackingKeypoint::Wrist_Center)] = &Wrist.Center;
-	EnumToTransformMap[static_cast<int32>(EHandTrackingKeypoint::Wrist_Ulnar)] = &Wrist.Ulnar;
-	EnumToTransformMap[static_cast<int32>(EHandTrackingKeypoint::Wrist_Radial)] = &Wrist.Radial;
+	EnumToTransformMap[static_cast<int32>(EMagicLeapHandTrackingKeypoint::Thumb_Tip)] = &Thumb.Tip;
+	EnumToTransformMap[static_cast<int32>(EMagicLeapHandTrackingKeypoint::Thumb_IP)] = &Thumb.DIP;
+	EnumToTransformMap[static_cast<int32>(EMagicLeapHandTrackingKeypoint::Thumb_MCP)] = &Thumb.PIP;
+	EnumToTransformMap[static_cast<int32>(EMagicLeapHandTrackingKeypoint::Thumb_CMC)] = &Thumb.MCP;
+
+	EnumToTransformMap[static_cast<int32>(EMagicLeapHandTrackingKeypoint::Index_Tip)] = &IndexFinger.Tip;
+	EnumToTransformMap[static_cast<int32>(EMagicLeapHandTrackingKeypoint::Index_DIP)] = &IndexFinger.DIP;
+	EnumToTransformMap[static_cast<int32>(EMagicLeapHandTrackingKeypoint::Index_PIP)] = &IndexFinger.PIP;
+	EnumToTransformMap[static_cast<int32>(EMagicLeapHandTrackingKeypoint::Index_MCP)] = &IndexFinger.MCP;
+
+	EnumToTransformMap[static_cast<int32>(EMagicLeapHandTrackingKeypoint::Middle_Tip)] = &MiddleFinger.Tip;
+	EnumToTransformMap[static_cast<int32>(EMagicLeapHandTrackingKeypoint::Middle_DIP)] = &MiddleFinger.DIP;
+	EnumToTransformMap[static_cast<int32>(EMagicLeapHandTrackingKeypoint::Middle_PIP)] = &MiddleFinger.PIP;
+	EnumToTransformMap[static_cast<int32>(EMagicLeapHandTrackingKeypoint::Middle_MCP)] = &MiddleFinger.MCP;
+
+	EnumToTransformMap[static_cast<int32>(EMagicLeapHandTrackingKeypoint::Ring_Tip)] = &RingFinger.Tip;
+	EnumToTransformMap[static_cast<int32>(EMagicLeapHandTrackingKeypoint::Ring_DIP)] = &RingFinger.DIP;
+	EnumToTransformMap[static_cast<int32>(EMagicLeapHandTrackingKeypoint::Ring_PIP)] = &RingFinger.PIP;
+	EnumToTransformMap[static_cast<int32>(EMagicLeapHandTrackingKeypoint::Ring_MCP)] = &RingFinger.MCP;
+
+	EnumToTransformMap[static_cast<int32>(EMagicLeapHandTrackingKeypoint::Pinky_Tip)] = &PinkyFinger.Tip;
+	EnumToTransformMap[static_cast<int32>(EMagicLeapHandTrackingKeypoint::Pinky_DIP)] = &PinkyFinger.DIP;
+	EnumToTransformMap[static_cast<int32>(EMagicLeapHandTrackingKeypoint::Pinky_PIP)] = &PinkyFinger.PIP;
+	EnumToTransformMap[static_cast<int32>(EMagicLeapHandTrackingKeypoint::Pinky_MCP)] = &PinkyFinger.MCP;
+
+	EnumToTransformMap[static_cast<int32>(EMagicLeapHandTrackingKeypoint::Wrist_Center)] = &Wrist.Center;
+	EnumToTransformMap[static_cast<int32>(EMagicLeapHandTrackingKeypoint::Wrist_Ulnar)] = &Wrist.Ulnar;
+	EnumToTransformMap[static_cast<int32>(EMagicLeapHandTrackingKeypoint::Wrist_Radial)] = &Wrist.Radial;
 	
-	EnumToTransformMap[static_cast<int32>(EHandTrackingKeypoint::Hand_Center)] = &HandCenter;
+	EnumToTransformMap[static_cast<int32>(EMagicLeapHandTrackingKeypoint::Hand_Center)] = &HandCenter;
 }
 
-bool FMagicLeapHandTracking::FHandState::GetTransform(EHandTrackingKeypoint KeyPoint, FTransform& OutTransform) const
+bool FMagicLeapHandTracking::FHandState::GetTransform(EMagicLeapHandTrackingKeypoint KeyPoint, FTransform& OutTransform) const
 {
-	check(static_cast<int32>(KeyPoint) >= 0 && static_cast<int32>(KeyPoint) <= EHandTrackingKeypointCount);
+	check(static_cast<int32>(KeyPoint) >= 0 && static_cast<int32>(KeyPoint) <= EMagicLeapHandTrackingKeypointCount);
 	OutTransform = EnumToTransformMap[static_cast<int32>(KeyPoint)]->Transform;
 	return IsValid();
 }
 
-const FMagicLeapHandTracking::FTransformRecord& FMagicLeapHandTracking::FHandState::GetTransform(EHandTrackingKeypoint KeyPoint) const
+const FMagicLeapHandTracking::FTransformRecord& FMagicLeapHandTracking::FHandState::GetTransform(EMagicLeapHandTrackingKeypoint KeyPoint) const
 {
-	check(static_cast<int32>(KeyPoint) >= 0 && static_cast<int32>(KeyPoint) <= EHandTrackingKeypointCount);
+	check(static_cast<int32>(KeyPoint) >= 0 && static_cast<int32>(KeyPoint) <= EMagicLeapHandTrackingKeypointCount);
 	return *EnumToTransformMap[static_cast<int32>(KeyPoint)];
 }
 
@@ -602,6 +607,24 @@ bool FMagicLeapHandTracking::GetControllerOrientationAndPosition(const int32 Con
 	return FXRMotionControllerBase::GetControllerOrientationAndPosition(ControllerIndex, MotionSource, OutOrientation, OutPosition, WorldToMetersScale);
 }
 
+ETrackingStatus FMagicLeapHandTracking::GetControllerTrackingStatus(const int32 ControllerIndex, const FName MotionSource) const
+{
+	ETrackingStatus status = ETrackingStatus::NotTracked;
+
+	if (ControllerIndex == DeviceIndex)
+	{
+		if (bIsHandTrackingStateValid)
+		{
+			const FHandState* HandState = SourceToHandMap.FindRef(MotionSource);
+			if (HandState != nullptr && HandState->IsValid())
+			{
+				status = ETrackingStatus::Tracked;
+			}
+		}
+	}
+	return status;
+}
+
 bool FMagicLeapHandTracking::GetControllerOrientationAndPosition(const int32 ControllerIndex, const EControllerHand DeviceHand, FRotator& OutOrientation, FVector& OutPosition, float WorldToMetersScale) const
 {
 	bool bControllerTracked = false;
@@ -649,24 +672,26 @@ bool FMagicLeapHandTracking::GetControllerOrientationAndPosition(const int32 Con
 ETrackingStatus FMagicLeapHandTracking::GetControllerTrackingStatus(const int32 ControllerIndex, const EControllerHand DeviceHand) const
 {
 	ETrackingStatus status = ETrackingStatus::NotTracked;
-	if (bIsHandTrackingStateValid)
+	if (ControllerIndex == DeviceIndex)
 	{
-		if (DeviceHand == EControllerHand::Special_1 || DeviceHand == EControllerHand::Special_3 || DeviceHand == EControllerHand::Special_5)
+		if (bIsHandTrackingStateValid)
 		{
-			if (LeftHand.Gesture != EHandTrackingGesture::NoHand)
+			if (DeviceHand == EControllerHand::Special_1 || DeviceHand == EControllerHand::Special_3 || DeviceHand == EControllerHand::Special_5)
 			{
-				status = ETrackingStatus::Tracked;
+				if (LeftHand.Gesture != EMagicLeapHandTrackingGesture::NoHand)
+				{
+					status = ETrackingStatus::Tracked;
+				}
 			}
-		}
-		else if (DeviceHand == EControllerHand::Special_2 || DeviceHand == EControllerHand::Special_4 || DeviceHand == EControllerHand::Special_6)
-		{
-			if (RightHand.Gesture != EHandTrackingGesture::NoHand)
+			else if (DeviceHand == EControllerHand::Special_2 || DeviceHand == EControllerHand::Special_4 || DeviceHand == EControllerHand::Special_6)
 			{
-				status = ETrackingStatus::Tracked;
+				if (RightHand.Gesture != EMagicLeapHandTrackingGesture::NoHand)
+				{
+					status = ETrackingStatus::Tracked;
+				}
 			}
 		}
 	}
-
 	return status;
 }
 
@@ -772,7 +797,7 @@ void FMagicLeapHandTracking::SendControllerEventsForHand(const MLHandTrackingHan
 	const float NewConfidence = NewHandState.keypose_confidence[NewHandState.keypose];
 	if (NewHandState.keypose != OldHandState.keypose)
 	{
-		FMagicLeapHMD::EnableInput EnableInputFromHMD;
+		MagicLeap::EnableInput EnableInputFromHMD;
 		// fixes unreferenced parameter error for Windows package builds.
 		(void)EnableInputFromHMD;
 		MessageHandler->OnControllerButtonReleased(GestureMap[OldHandState.keypose], DeviceIndex, false);
@@ -783,7 +808,7 @@ void FMagicLeapHandTracking::SendControllerEventsForHand(const MLHandTrackingHan
 	}
 	else if (OldConfidence < GestureConfidenceThresholds[GestureIndex] && NewConfidence >= GestureConfidenceThresholds[GestureIndex])
 	{
-		FMagicLeapHMD::EnableInput EnableInputFromHMD;
+		MagicLeap::EnableInput EnableInputFromHMD;
 		// fixes unreferenced parameter error for Windows package builds.
 		(void)EnableInputFromHMD;
 		MessageHandler->OnControllerButtonPressed(GestureMap[GestureIndex], DeviceIndex, false);
@@ -810,34 +835,24 @@ bool FMagicLeapHandTracking::IsGamepadAttached() const
 #endif //WITH_MLSDK
 }
 
-void FMagicLeapHandTracking::Enable()
-{
-	checkf(false, TEXT("FMagicLeapHandTracking::Enable is not supported! Check 'SupportsExplicitEnable()' first!"));
-}
-
-bool FMagicLeapHandTracking::SupportsExplicitEnable() const
-{
-	return false;
-}
-
-void FMagicLeapHandTracking::Disable()
+void FMagicLeapHandTracking::DestroyEntityTracker()
 {
 #if WITH_MLSDK
 	if (MLHandleIsValid(HandTracker))
 	{
-		if (MLResult_Ok == MLHandTrackingDestroy(HandTracker))
+		MLResult Result = MLHandTrackingDestroy(HandTracker);
+		if (Result != MLResult_Ok)
 		{
-			HandTracker = ML_INVALID_HANDLE;
-		}
-		else
-		{
-			UE_LOG(LogMagicLeapHandTracking, Error, TEXT("Error destroying hand tracker."));
+			UE_LOG(LogMagicLeapHandTracking, Error, TEXT("MLHandTrackingDestroy failed with error '%s'."), UTF8_TO_TCHAR(MLGetResultString(Result)));
+			return;
 		}
+
+		HandTracker = ML_INVALID_HANDLE;
 	}
 #endif
 }
 
-void FMagicLeapHandTracking::OnBeginRendering_GameThread_Update()
+void FMagicLeapHandTracking::OnBeginRendering_GameThread()
 {
 #if WITH_MLSDK
 	UpdateCurrentHandTrackerTransforms();
@@ -859,14 +874,14 @@ bool FMagicLeapHandTracking::IsHandTrackingStateValid() const
 	return bIsHandTrackingStateValid;
 }
 
-bool FMagicLeapHandTracking::GetKeypointTransform(EControllerHand Hand, EHandTrackingKeypoint Keypoint, FTransform& OutTransform) const
+bool FMagicLeapHandTracking::GetKeypointTransform(EControllerHand Hand, EMagicLeapHandTrackingKeypoint Keypoint, FTransform& OutTransform) const
 {
 	const FMagicLeapHandTracking::FHandState& HandState = (Hand == EControllerHand::Left) ? GetLeftHandState() : GetRightHandState();
 
 	return HandState.GetTransform(Keypoint, OutTransform);
 }
 
-bool FMagicLeapHandTracking::SetConfiguration(bool bTrackingEnabled, const TArray<EHandTrackingGesture>& ActiveKeyPoses, EHandTrackingKeypointFilterLevel KeypointsFilterLevel, EHandTrackingGestureFilterLevel GestureFilterLevel)
+bool FMagicLeapHandTracking::SetConfiguration(bool bTrackingEnabled, const TArray<EMagicLeapHandTrackingGesture>& ActiveKeyPoses, EMagicLeapHandTrackingKeypointFilterLevel KeypointsFilterLevel, EMagicLeapHandTrackingGestureFilterLevel GestureFilterLevel)
 {
 #if WITH_MLSDK
 	ConditionallyEnable();
@@ -886,7 +901,7 @@ bool FMagicLeapHandTracking::SetConfiguration(bool bTrackingEnabled, const TArra
 
 	bool EnableHandTrackingPipeline = false;
 
-	for (const EHandTrackingGesture& StaticPose : ActiveKeyPoses)
+	for (const EMagicLeapHandTrackingGesture& StaticPose : ActiveKeyPoses)
 	{
 		if (static_cast<uint32>(StaticPose) <= MaxConfigurablePose)
 		{
@@ -897,14 +912,16 @@ bool FMagicLeapHandTracking::SetConfiguration(bool bTrackingEnabled, const TArra
 	Config.handtracking_pipeline_enabled = bTrackingEnabled;
 	Config.keypoints_filter_level = UnrealToMLKeypointsFilterLevel(KeypointsFilterLevel);
 	Config.pose_filter_level = UnrealToMLGestureFilterLevel(GestureFilterLevel);
-
-	return MLHandTrackingSetConfiguration(HandTracker, &Config) == MLResult_Ok;
+	MLResult Result = MLHandTrackingSetConfiguration(HandTracker, &Config);
+	UE_CLOG(Result != MLResult_Ok, LogMagicLeapHandTracking, Error, TEXT("MLHandTrackingSetConfiguration failed with error '%s'."), UTF8_TO_TCHAR(MLGetResultString(Result)));
+	
+	return Result == MLResult_Ok;
 #else
 	return false;
 #endif //WITH_MLSDK
 }
 
-bool FMagicLeapHandTracking::GetConfiguration(bool& bTrackingEnabled, TArray<EHandTrackingGesture>& ActiveKeyPoses, EHandTrackingKeypointFilterLevel& KeypointsFilterLevel, EHandTrackingGestureFilterLevel& GestureFilterLevel)
+bool FMagicLeapHandTracking::GetConfiguration(bool& bTrackingEnabled, TArray<EMagicLeapHandTrackingGesture>& ActiveKeyPoses, EMagicLeapHandTrackingKeypointFilterLevel& KeypointsFilterLevel, EMagicLeapHandTrackingGestureFilterLevel& GestureFilterLevel)
 {
 #if WITH_MLSDK
 	ConditionallyEnable();
@@ -914,43 +931,45 @@ bool FMagicLeapHandTracking::GetConfiguration(bool& bTrackingEnabled, TArray<EHa
 		return false;
 	}
 
-	MLHandTrackingConfiguration config;
-	bool validResult = MLHandTrackingGetConfiguration(HandTracker, &config) == MLResult_Ok;
+	MLHandTrackingConfiguration Config;
+	MLResult Result = MLHandTrackingGetConfiguration(HandTracker, &Config);
+	if (Result != MLResult_Ok)
+	{
+		UE_LOG(LogMagicLeapHandTracking, Error, TEXT("MLHandTrackingGetConfiguration failed with error '%s'."), UTF8_TO_TCHAR(MLGetResultString(Result)));
+		return false;
+	}
 
-	if (validResult)
+	ActiveKeyPoses.Empty(MLHandTrackingKeyPose_Count);
+	for (uint32 i = 0; i < MLHandTrackingKeyPose_Count; ++i)
 	{
-		ActiveKeyPoses.Empty(MLHandTrackingKeyPose_Count);
-		for (uint32 i = 0; i < MLHandTrackingKeyPose_Count; ++i)
+		if (Config.keypose_config[i])
 		{
-			if (config.keypose_config[i])
-			{
-				ActiveKeyPoses.Add(static_cast<EHandTrackingGesture>(i));
-			}
+			ActiveKeyPoses.Add(static_cast<EMagicLeapHandTrackingGesture>(i));
 		}
-
-		bTrackingEnabled = config.handtracking_pipeline_enabled;
-		KeypointsFilterLevel = MLToUnrealKeypointsFilterLevel(config.keypoints_filter_level);
-		GestureFilterLevel = MLToUnrealGestureFilterLevel(config.pose_filter_level);
 	}
 
-	return validResult;
+	bTrackingEnabled = Config.handtracking_pipeline_enabled;
+	KeypointsFilterLevel = MLToUnrealKeypointsFilterLevel(Config.keypoints_filter_level);
+	GestureFilterLevel = MLToUnrealGestureFilterLevel(Config.pose_filter_level);
+	
+	return true;
 #else
 	return false;
 #endif //WITH_MLSDK
 
 }
 
-void FMagicLeapHandTracking::SetGestureConfidenceThreshold(EHandTrackingGesture Gesture, float Confidence)
+void FMagicLeapHandTracking::SetGestureConfidenceThreshold(EMagicLeapHandTrackingGesture Gesture, float Confidence)
 {
-	if (Gesture <= EHandTrackingGesture::NoHand)
+	if ((Gesture <= EMagicLeapHandTrackingGesture::NoHand) && (GestureConfidenceThresholds.Num() != 0))
 	{
 		GestureConfidenceThresholds[static_cast<int32>(Gesture)] = Confidence;
 	}
 }
 
-float FMagicLeapHandTracking::GetGestureConfidenceThreshold(EHandTrackingGesture Gesture) const
+float FMagicLeapHandTracking::GetGestureConfidenceThreshold(EMagicLeapHandTrackingGesture Gesture) const
 {
-	if (Gesture <= EHandTrackingGesture::NoHand)
+	if ((Gesture <= EMagicLeapHandTrackingGesture::NoHand) && (GestureConfidenceThresholds.Num() != 0))
 	{
 		return GestureConfidenceThresholds[static_cast<int32>(Gesture)];
 	}
@@ -960,18 +979,18 @@ float FMagicLeapHandTracking::GetGestureConfidenceThreshold(EHandTrackingGesture
 namespace MagicLeapHandTracking
 {
 #if WITH_MLSDK
-	bool FetchTransform(const FAppFramework& AppFramework, const MLKeyPointState& Source, FMagicLeapHandTracking::FTransformRecord& OutDest, const TCHAR* DebugString, const TCHAR* DebugString2, const TCHAR* DebugString3)
+	bool FetchTransform(const IMagicLeapPlugin& MLPlugin, const MLKeyPointState& Source, FMagicLeapHandTracking::FTransformRecord& OutDest, const TCHAR* DebugString, const TCHAR* DebugString2, const TCHAR* DebugString3)
 	{
 		bool bResult = false;
 
 		if (Source.is_valid)
 		{
-			EFailReason FailReason = EFailReason::None;
+			EMagicLeapTransformFailReason FailReason = EMagicLeapTransformFailReason::None;
 			FTransform Transform;
-			bResult = AppFramework.GetTransform(Source.frame_id, Transform, FailReason);
+			bResult = MLPlugin.GetTransform(Source.frame_id, Transform, FailReason);
 			if (!bResult)
 			{
-				if (FailReason == EFailReason::NaNsInTransform)
+				if (FailReason == EMagicLeapTransformFailReason::NaNsInTransform)
 				{
 					UE_LOG(LogMagicLeapHandTracking, Error, TEXT("NaNs in %s %s %s transform."), DebugString, DebugString2, DebugString3);
 				}
@@ -986,31 +1005,31 @@ namespace MagicLeapHandTracking
 		return bResult;
 	}
 
-	void FetchFingerTransforms(const FAppFramework& AppFramework, const MLFingerState& Source, FMagicLeapHandTracking::FDigitTransforms& OutDest, const TCHAR* DebugString, const TCHAR* DebugString2)
+	void FetchFingerTransforms(const IMagicLeapPlugin& MLPlugin, const MLFingerState& Source, FMagicLeapHandTracking::FDigitTransforms& OutDest, const TCHAR* DebugString, const TCHAR* DebugString2)
 	{
-		FetchTransform(AppFramework, Source.tip, OutDest.Tip, DebugString, DebugString2, TEXT("Tip"));
-		FetchTransform(AppFramework, Source.dip, OutDest.DIP, DebugString, DebugString2, TEXT("DIP"));
-		FetchTransform(AppFramework, Source.pip, OutDest.PIP, DebugString, DebugString2, TEXT("PIP"));
-		FetchTransform(AppFramework, Source.mcp, OutDest.MCP, DebugString, DebugString2, TEXT("MCP"));
+		FetchTransform(MLPlugin, Source.tip, OutDest.Tip, DebugString, DebugString2, TEXT("Tip"));
+		FetchTransform(MLPlugin, Source.dip, OutDest.DIP, DebugString, DebugString2, TEXT("DIP"));
+		FetchTransform(MLPlugin, Source.pip, OutDest.PIP, DebugString, DebugString2, TEXT("PIP"));
+		FetchTransform(MLPlugin, Source.mcp, OutDest.MCP, DebugString, DebugString2, TEXT("MCP"));
 	}
 
-	void FetchHandTransforms(const FAppFramework& AppFramework, const MLHandTrackingStaticHandState& Source, FMagicLeapHandTracking::FHandState& OutDest, const TCHAR* DebugString)
+	void FetchHandTransforms(const IMagicLeapPlugin& MLPlugin, const MLHandTrackingStaticHandState& Source, FMagicLeapHandTracking::FHandState& OutDest, const TCHAR* DebugString)
 	{
-		FetchTransform(AppFramework, Source.thumb.tip, OutDest.Thumb.Tip, DebugString, TEXT("thumb"), TEXT("Tip"));
-		FetchTransform(AppFramework, Source.thumb.ip,  OutDest.Thumb.DIP, DebugString, TEXT("thumb"), TEXT("DIP"));
-		FetchTransform(AppFramework, Source.thumb.mcp, OutDest.Thumb.PIP, DebugString, TEXT("thumb"), TEXT("PIP"));
-		FetchTransform(AppFramework, Source.thumb.cmc, OutDest.Thumb.MCP, DebugString, TEXT("thumb"), TEXT("MCP"));
+		FetchTransform(MLPlugin, Source.thumb.tip, OutDest.Thumb.Tip, DebugString, TEXT("thumb"), TEXT("Tip"));
+		FetchTransform(MLPlugin, Source.thumb.ip,  OutDest.Thumb.DIP, DebugString, TEXT("thumb"), TEXT("DIP"));
+		FetchTransform(MLPlugin, Source.thumb.mcp, OutDest.Thumb.PIP, DebugString, TEXT("thumb"), TEXT("PIP"));
+		FetchTransform(MLPlugin, Source.thumb.cmc, OutDest.Thumb.MCP, DebugString, TEXT("thumb"), TEXT("MCP"));
 
-		FetchFingerTransforms(AppFramework, Source.index,  OutDest.IndexFinger,  DebugString, TEXT("Index"));
-		FetchFingerTransforms(AppFramework, Source.middle, OutDest.MiddleFinger, DebugString, TEXT("Index"));
-		FetchFingerTransforms(AppFramework, Source.ring,   OutDest.RingFinger,   DebugString, TEXT("Index"));
-		FetchFingerTransforms(AppFramework, Source.pinky,  OutDest.PinkyFinger,  DebugString, TEXT("Index"));
+		FetchFingerTransforms(MLPlugin, Source.index,  OutDest.IndexFinger,  DebugString, TEXT("Index"));
+		FetchFingerTransforms(MLPlugin, Source.middle, OutDest.MiddleFinger, DebugString, TEXT("Index"));
+		FetchFingerTransforms(MLPlugin, Source.ring,   OutDest.RingFinger,   DebugString, TEXT("Index"));
+		FetchFingerTransforms(MLPlugin, Source.pinky,  OutDest.PinkyFinger,  DebugString, TEXT("Index"));
 
-		FetchTransform(AppFramework, Source.wrist.center, OutDest.Wrist.Center, DebugString, TEXT("wrist"), TEXT("center"));
-		FetchTransform(AppFramework, Source.wrist.ulnar,  OutDest.Wrist.Ulnar,  DebugString, TEXT("wrist"), TEXT("ulnar"));
-		FetchTransform(AppFramework, Source.wrist.radial, OutDest.Wrist.Radial, DebugString, TEXT("wrist"), TEXT("radial"));
+		FetchTransform(MLPlugin, Source.wrist.center, OutDest.Wrist.Center, DebugString, TEXT("wrist"), TEXT("center"));
+		FetchTransform(MLPlugin, Source.wrist.ulnar,  OutDest.Wrist.Ulnar,  DebugString, TEXT("wrist"), TEXT("ulnar"));
+		FetchTransform(MLPlugin, Source.wrist.radial, OutDest.Wrist.Radial, DebugString, TEXT("wrist"), TEXT("radial"));
 
-		FetchTransform(AppFramework, Source.hand_center, OutDest.HandCenter, DebugString, TEXT(""), TEXT("center"));
+		FetchTransform(MLPlugin, Source.hand_center, OutDest.HandCenter, DebugString, TEXT(""), TEXT("center"));
 	}
 
 	//void LogMLHandTackingHandState(const MLHandTrackingHandState& State)
@@ -1109,8 +1128,16 @@ void FMagicLeapHandTracking::UpdateTrackerData()
 	if (MLHandleIsValid(HandTracker))
 	{
 		CurrentHandTrackingDataIndex = 1 - CurrentHandTrackingDataIndex;
-		bIsHandTrackingStateValid = MLHandTrackingGetData(HandTracker, &HandTrackingDatas[CurrentHandTrackingDataIndex]) == MLResult_Ok;
-		//MagicLeapHandTracking::LogMLHandTrackingData(GetCurrentHandTrackingData());
+		MLResult Result = MLHandTrackingGetData(HandTracker, &HandTrackingDatas[CurrentHandTrackingDataIndex]);
+		if (Result != MLResult_Ok)
+		{
+			bIsHandTrackingStateValid = false;
+			UE_LOG(LogMagicLeapHandTracking, Error, TEXT("MLHandTrackingGetData failed with error '%s'."), UTF8_TO_TCHAR(MLGetResultString(Result)));
+		}
+		else
+		{
+			bIsHandTrackingStateValid = true;
+		}
 	}
 	else
 	{
@@ -1118,9 +1145,6 @@ void FMagicLeapHandTracking::UpdateTrackerData()
 	}
 	if (bIsHandTrackingStateValid && IMagicLeapPlugin::Get().IsMagicLeapHMDValid())
 	{
-		const FAppFramework& AppFramework = static_cast<FMagicLeapHMD*>(GEngine->XRSystem->GetHMDDevice())->GetAppFrameworkConst();
-		check(AppFramework.IsInitialized());
-
 		const MLHandTrackingData& HandTrackingData = GetCurrentHandTrackingData();
 
 		LeftHand.Gesture = TranslateGestureEnum(HandTrackingData.left_hand_state.keypose);
@@ -1139,12 +1163,11 @@ void FMagicLeapHandTracking::UpdateCurrentHandTrackerTransforms()
 
 	if (MLHandleIsValid(HandTracker) && bIsHandTrackingStateValid && IMagicLeapPlugin::Get().IsMagicLeapHMDValid())
 	{
-		const FAppFramework& AppFramework = static_cast<FMagicLeapHMD*>(GEngine->XRSystem->GetHMDDevice())->GetAppFrameworkConst();
-		check(AppFramework.IsInitialized());
+		IMagicLeapPlugin& MLPlugin = IMagicLeapPlugin::Get();
 
 		const MLHandTrackingData& HandTrackingData = GetCurrentHandTrackingData();
 
-		if (LeftHand.Gesture != EHandTrackingGesture::NoHand)
+		if (LeftHand.Gesture != EMagicLeapHandTrackingGesture::NoHand)
 		{
 			LeftHand.HandCenterNormalized = MagicLeap::ToFVector(HandTrackingData.left_hand_state.hand_center_normalized, 1.0f);
 			if (LeftHand.HandCenterNormalized.ContainsNaN())
@@ -1153,9 +1176,9 @@ void FMagicLeapHandTracking::UpdateCurrentHandTrackerTransforms()
 				LeftHand.HandCenterNormalized = FVector::ZeroVector;
 			}
 
-			MagicLeapHandTracking::FetchHandTransforms(AppFramework, HandTrackingStaticData.left, LeftHand, TEXT("left hand"));
+			MagicLeapHandTracking::FetchHandTransforms(MLPlugin, HandTrackingStaticData.left, LeftHand, TEXT("left hand"));
 		}
-		if (RightHand.Gesture != EHandTrackingGesture::NoHand)
+		if (RightHand.Gesture != EMagicLeapHandTrackingGesture::NoHand)
 		{
 			RightHand.HandCenterNormalized = MagicLeap::ToFVector(HandTrackingData.right_hand_state.hand_center_normalized, 1.0f);
 			if (RightHand.HandCenterNormalized.ContainsNaN())
@@ -1164,7 +1187,7 @@ void FMagicLeapHandTracking::UpdateCurrentHandTrackerTransforms()
 				RightHand.HandCenterNormalized = FVector::ZeroVector;
 			}
 
-			MagicLeapHandTracking::FetchHandTransforms(AppFramework, HandTrackingStaticData.right, RightHand, TEXT("right hand"));
+			MagicLeapHandTracking::FetchHandTransforms(MLPlugin, HandTrackingStaticData.right, RightHand, TEXT("right hand"));
 		}
 	}
 #endif
@@ -1178,7 +1201,9 @@ void FMagicLeapHandTracking::AddKeys()
 	EKeys::AddKey(FKeyDetails(FStaticHandTracking::Left_Pinch, LOCTEXT("MagicLeap_Left_Pinch", "ML Left Pinch"), FKeyDetails::GamepadKey));
 	EKeys::AddKey(FKeyDetails(FStaticHandTracking::Left_Thumb, LOCTEXT("MagicLeap_Left_Thumb", "ML Left Thumb"), FKeyDetails::GamepadKey));
 	EKeys::AddKey(FKeyDetails(FStaticHandTracking::Left_L, LOCTEXT("MagicLeap_Left_L", "ML Left L"), FKeyDetails::GamepadKey));
-	EKeys::AddKey(FKeyDetails(FStaticHandTracking::Left_OpenHandBack, LOCTEXT("MagicLeap_Left_OpenHandBack", "ML Left Open Hand Back"), FKeyDetails::GamepadKey));
+	EKeys::AddKey(FKeyDetails(FStaticHandTracking::Left_OpenHand, LOCTEXT("MagicLeap_Left_OpenHand", "ML Left Open Hand"), FKeyDetails::GamepadKey));
+	// DEPRECATED
+	EKeys::AddKey(FKeyDetails(FStaticHandTracking::Left_OpenHandBack, LOCTEXT("MagicLeap_Left_OpenHandBack", "ML Left Open Hand Back DEPRECATED"), FKeyDetails::GamepadKey));
 	EKeys::AddKey(FKeyDetails(FStaticHandTracking::Left_Ok, LOCTEXT("MagicLeap_Left_Ok", "ML Left Ok"), FKeyDetails::GamepadKey));
 	EKeys::AddKey(FKeyDetails(FStaticHandTracking::Left_C, LOCTEXT("MagicLeap_Left_C", "ML Left C"), FKeyDetails::GamepadKey));
 	EKeys::AddKey(FKeyDetails(FStaticHandTracking::Left_NoPose, LOCTEXT("MagicLeap_Left_NoPose", "ML Left NoPose"), FKeyDetails::GamepadKey));
@@ -1190,7 +1215,9 @@ void FMagicLeapHandTracking::AddKeys()
 	EKeys::AddKey(FKeyDetails(FStaticHandTracking::Right_Pinch, LOCTEXT("MagicLeap_Right_Pinch", "ML Right Pinch"), FKeyDetails::GamepadKey));
 	EKeys::AddKey(FKeyDetails(FStaticHandTracking::Right_Thumb, LOCTEXT("MagicLeap_Right_Thumb", "ML Right Thumb"), FKeyDetails::GamepadKey));
 	EKeys::AddKey(FKeyDetails(FStaticHandTracking::Right_L, LOCTEXT("MagicLeap_Right_L", "ML Right L"), FKeyDetails::GamepadKey));
-	EKeys::AddKey(FKeyDetails(FStaticHandTracking::Right_OpenHandBack, LOCTEXT("MagicLeap_Right_OpenHandBack", "ML Right Open Hand Back"), FKeyDetails::GamepadKey));
+	EKeys::AddKey(FKeyDetails(FStaticHandTracking::Right_OpenHand, LOCTEXT("MagicLeap_Right_OpenHand", "ML Right Open Hand"), FKeyDetails::GamepadKey));
+	// DEPRECATED
+	EKeys::AddKey(FKeyDetails(FStaticHandTracking::Right_OpenHandBack, LOCTEXT("MagicLeap_Right_OpenHandBack", "ML Right Open Hand Back DEPRECATED"), FKeyDetails::GamepadKey));
 	EKeys::AddKey(FKeyDetails(FStaticHandTracking::Right_Ok, LOCTEXT("MagicLeap_Right_Ok", "ML Right Ok"), FKeyDetails::GamepadKey));
 	EKeys::AddKey(FKeyDetails(FStaticHandTracking::Right_C, LOCTEXT("MagicLeap_Right_C", "ML Right C"), FKeyDetails::GamepadKey));
 	EKeys::AddKey(FKeyDetails(FStaticHandTracking::Right_NoPose, LOCTEXT("MagicLeap_Right_NoPose", "ML Right NoPose"), FKeyDetails::GamepadKey));
@@ -1209,6 +1236,8 @@ void FMagicLeapHandTracking::AddKeys()
 	LeftStaticGestureMap[MLHandTrackingKeyPose_Pinch] = FMagicLeapGestureKeyNames::Left_Pinch_Name;
 	LeftStaticGestureMap[MLHandTrackingKeyPose_Thumb] = FMagicLeapGestureKeyNames::Left_Thumb_Name;
 	LeftStaticGestureMap[MLHandTrackingKeyPose_L] = FMagicLeapGestureKeyNames::Left_L_Name;
+	LeftStaticGestureMap[MLHandTrackingKeyPose_OpenHand] = FMagicLeapGestureKeyNames::Left_OpenHand_Name;
+	// DEPRECATED
 	LeftStaticGestureMap[MLHandTrackingKeyPose_OpenHandBack] = FMagicLeapGestureKeyNames::Left_OpenHandBack_Name;
 	LeftStaticGestureMap[MLHandTrackingKeyPose_Ok] = FMagicLeapGestureKeyNames::Left_Ok_Name;
 	LeftStaticGestureMap[MLHandTrackingKeyPose_C] = FMagicLeapGestureKeyNames::Left_C_Name;
@@ -1221,6 +1250,8 @@ void FMagicLeapHandTracking::AddKeys()
 	RightStaticGestureMap[MLHandTrackingKeyPose_Pinch] = FMagicLeapGestureKeyNames::Right_Pinch_Name;
 	RightStaticGestureMap[MLHandTrackingKeyPose_Thumb] = FMagicLeapGestureKeyNames::Right_Thumb_Name;
 	RightStaticGestureMap[MLHandTrackingKeyPose_L] = FMagicLeapGestureKeyNames::Right_L_Name;
+	RightStaticGestureMap[MLHandTrackingKeyPose_OpenHand] = FMagicLeapGestureKeyNames::Right_OpenHand_Name;
+	// DEPRECATED
 	RightStaticGestureMap[MLHandTrackingKeyPose_OpenHandBack] = FMagicLeapGestureKeyNames::Right_OpenHandBack_Name;
 	RightStaticGestureMap[MLHandTrackingKeyPose_Ok] = FMagicLeapGestureKeyNames::Right_Ok_Name;
 	RightStaticGestureMap[MLHandTrackingKeyPose_C] = FMagicLeapGestureKeyNames::Right_C_Name;
@@ -1231,33 +1262,33 @@ void FMagicLeapHandTracking::AddKeys()
 }
 
 #if WITH_MLSDK
-EHandTrackingGesture FMagicLeapHandTracking::TranslateGestureEnum(MLHandTrackingKeyPose HandState)
+EMagicLeapHandTrackingGesture FMagicLeapHandTracking::TranslateGestureEnum(MLHandTrackingKeyPose HandState)
 {
 	switch (HandState)
 	{
 	case MLHandTrackingKeyPose_Finger:
-		return EHandTrackingGesture::Finger;
+		return EMagicLeapHandTrackingGesture::Finger;
 	case MLHandTrackingKeyPose_Fist:
-		return EHandTrackingGesture::Fist;
+		return EMagicLeapHandTrackingGesture::Fist;
 	case MLHandTrackingKeyPose_Pinch:
-		return EHandTrackingGesture::Pinch;
+		return EMagicLeapHandTrackingGesture::Pinch;
 	case MLHandTrackingKeyPose_Thumb:
-		return EHandTrackingGesture::Thumb;
+		return EMagicLeapHandTrackingGesture::Thumb;
 	case MLHandTrackingKeyPose_L:
-		return EHandTrackingGesture::L;
-	case MLHandTrackingKeyPose_OpenHandBack:
-		return EHandTrackingGesture::OpenHandBack;
+		return EMagicLeapHandTrackingGesture::L;
+	case MLHandTrackingKeyPose_OpenHand:
+		return EMagicLeapHandTrackingGesture::OpenHand;
 	case MLHandTrackingKeyPose_Ok:
-		return EHandTrackingGesture::Ok;
+		return EMagicLeapHandTrackingGesture::Ok;
 	case MLHandTrackingKeyPose_C:
-		return EHandTrackingGesture::C;
+		return EMagicLeapHandTrackingGesture::C;
 	case MLHandTrackingKeyPose_NoPose:
-		return EHandTrackingGesture::NoPose;
+		return EMagicLeapHandTrackingGesture::NoPose;
 	case MLHandTrackingKeyPose_NoHand:
-		return EHandTrackingGesture::NoHand;
+		return EMagicLeapHandTrackingGesture::NoHand;
 	default:
 		check(false);
-		return EHandTrackingGesture::NoHand;
+		return EMagicLeapHandTrackingGesture::NoHand;
 	}
 }
 #endif //WITH_MLSDK
@@ -1265,20 +1296,26 @@ EHandTrackingGesture FMagicLeapHandTracking::TranslateGestureEnum(MLHandTracking
 void FMagicLeapHandTracking::ConditionallyEnable()
 {
 #if WITH_MLSDK
-	if (!MLHandleIsValid(HandTracker) && GEngine->XRSystem.IsValid() && GEngine->XRSystem->GetHMDDevice())
+	if (!MLHandleIsValid(HandTracker) && IMagicLeapPlugin::Get().IsPerceptionEnabled())
 	{
-		MLResult CreateResult = MLHandTrackingCreate(&HandTracker);
-		if (CreateResult == MLResult_Ok && MLHandleIsValid(HandTracker))
+		MLResult Result = MLHandTrackingCreate(&HandTracker);
+		if (Result != MLResult_Ok)
 		{
-			if (MLResult_Ok != MLHandTrackingGetStaticData(HandTracker, &HandTrackingStaticData))
-			{
-				UE_LOG(LogMagicLeapHandTracking, Error, TEXT("Error getting hand tracker static data."));
-			}
-			//MagicLeapHandTracking::LogMLHandTrackingStaticData(HandTrackingStaticData);
+			UE_LOG(LogMagicLeapHandTracking, Error, TEXT("MLHandTrackingCreate failed with error '%s'."), UTF8_TO_TCHAR(MLGetResultString(Result)));
+			return;
 		}
-		else
+
+		if (!MLHandleIsValid(HandTracker))
+		{
+			UE_LOG(LogMagicLeapHandTracking, Error, TEXT("MLHandTrackingCreate succeeded but resulting tracker is invalid."));
+			return;
+		}
+
+		Result = MLHandTrackingGetStaticData(HandTracker, &HandTrackingStaticData);
+		if (Result != MLResult_Ok)
 		{
-			UE_LOG(LogMagicLeapHandTracking, Error, TEXT("Error creating hand tracker."));
+			UE_LOG(LogMagicLeapHandTracking, Error, TEXT("MLHandTrackingGetStaticData failed with error '%s'."), UTF8_TO_TCHAR(MLGetResultString(Result)));
+			return;
 		}
 	}
 #endif //WITH_MLSDK
@@ -1289,39 +1326,33 @@ void FMagicLeapHandTracking::OnAppPause()
 #if WITH_MLSDK
 	if (!MLHandleIsValid(HandTracker))
 	{
-		UE_LOG(LogMagicLeapHandTracking, Error, TEXT("Hand tracker was invalid on application pause."));
+		return;
 	}
-	else
-	{
-		MLHandTrackingConfiguration HandTrackingConfig;
-
-		if (MLResult_Ok != MLHandTrackingGetConfiguration(HandTracker, &HandTrackingConfig))
-		{
-			UE_LOG(LogMagicLeapHandTracking, Error, TEXT("Failed to retrieve hand tracking configuration on application pause."));
-		}
-		else
-		{
-			bWasSystemEnabledOnPause = HandTrackingConfig.handtracking_pipeline_enabled;
 
-			if (!bWasSystemEnabledOnPause)
-			{
-				UE_LOG(LogMagicLeapHandTracking, Log, TEXT("Hand tracking was not enabled at time of application pause."));
-			}
-			else
-			{
-				HandTrackingConfig.handtracking_pipeline_enabled = false;
+	MLHandTrackingConfiguration HandTrackingConfig;
+	MLResult Result = MLHandTrackingGetConfiguration(HandTracker, &HandTrackingConfig);
+	if (Result != MLResult_Ok)
+	{
+		UE_LOG(LogMagicLeapHandTracking, Error, TEXT("MLHandTrackingGetConfiguration failed with error '%s'."), UTF8_TO_TCHAR(MLGetResultString(Result)));
+		return;
+	}
 
-				if (MLResult_Ok != MLHandTrackingSetConfiguration(HandTracker, &HandTrackingConfig))
-				{
-					UE_LOG(LogMagicLeapHandTracking, Error, TEXT("Failed to disable hand tracking on application pause."));
-				}
-				else
-				{
-					UE_LOG(LogMagicLeapHandTracking, Log, TEXT("Hand tracking paused until app resumes."));
-				}
-			}
-		}
+	bWasSystemEnabledOnPause = HandTrackingConfig.handtracking_pipeline_enabled;
+	if (!bWasSystemEnabledOnPause)
+	{
+		UE_LOG(LogMagicLeapHandTracking, Log, TEXT("Hand tracking was not enabled at time of application pause."));
+		return;
+	}
+	
+	HandTrackingConfig.handtracking_pipeline_enabled = false;
+	Result = MLHandTrackingSetConfiguration(HandTracker, &HandTrackingConfig);
+	if (Result != MLResult_Ok)
+	{
+		UE_LOG(LogMagicLeapHandTracking, Error, TEXT("MLHandTrackingSetConfiguration failed with error '%s'."), UTF8_TO_TCHAR(MLGetResultString(Result)));
+		return;
 	}
+
+	UE_LOG(LogMagicLeapHandTracking, Log, TEXT("Hand tracking paused until app resumes."));
 #endif //WITH_MLSDK
 }
 
@@ -1330,37 +1361,32 @@ void FMagicLeapHandTracking::OnAppResume()
 #if WITH_MLSDK
 	if (!MLHandleIsValid(HandTracker))
 	{
-		UE_LOG(LogMagicLeapHandTracking, Error, TEXT("Hand tracker was invalid on application resume."));
+		return;
 	}
-	else
+
+	MLHandTrackingConfiguration HandTrackingConfig;
+	if (!bWasSystemEnabledOnPause)
 	{
-		MLHandTrackingConfiguration HandTrackingConfig;
+		UE_LOG(LogMagicLeapHandTracking, Log, TEXT("Not resuming hand tracking as it was not enabled at time of application pause."));
+		return;
+	}
 
-		if (!bWasSystemEnabledOnPause)
-		{
-			UE_LOG(LogMagicLeapHandTracking, Log, TEXT("Not resuming hand tracking as it was not enabled at time of application pause."));
-		}
-		else
-		{
-			if (MLResult_Ok != MLHandTrackingGetConfiguration(HandTracker, &HandTrackingConfig))
-			{
-				UE_LOG(LogMagicLeapHandTracking, Error, TEXT("Failed to retrieve hand tracking configuration on application resume."));
-			}
-			else
-			{
-				HandTrackingConfig.handtracking_pipeline_enabled = true;
+	MLResult Result = MLHandTrackingGetConfiguration(HandTracker, &HandTrackingConfig);
+	if (Result != MLResult_Ok)
+	{
+		UE_LOG(LogMagicLeapHandTracking, Error, TEXT("MLHandTrackingGetConfiguration failed with error '%s'."), UTF8_TO_TCHAR(MLGetResultString(Result)));
+		return;
+	}
 
-				if (MLResult_Ok != MLHandTrackingSetConfiguration(HandTracker, &HandTrackingConfig))
-				{
-					UE_LOG(LogMagicLeapHandTracking, Error, TEXT("Failed to re-enable hand tracking on application resume."));
-				}
-				else
-				{
-					UE_LOG(LogMagicLeapHandTracking, Log, TEXT("Hand tracking resumed."));
-				}
-			}
-		}
+	HandTrackingConfig.handtracking_pipeline_enabled = true;
+	Result = MLHandTrackingSetConfiguration(HandTracker, &HandTrackingConfig);
+	if (Result != MLResult_Ok)
+	{
+		UE_LOG(LogMagicLeapHandTracking, Error, TEXT("MLHandTrackingSetConfiguration failed with error '%s'."), UTF8_TO_TCHAR(MLGetResultString(Result)));
+		return;
 	}
+
+	UE_LOG(LogMagicLeapHandTracking, Log, TEXT("Hand tracking resumed."));
 #endif //WITH_MLSDK
 }
 
diff --git a/Engine/Plugins/Lumin/MagicLeap/Source/MagicLeapHandTracking/Private/MagicLeapHandTracking.h b/Engine/Plugins/Lumin/MagicLeap/Source/MagicLeapHandTracking/Private/MagicLeapHandTracking.h
--- a/Engine/Plugins/Lumin/MagicLeap/Source/MagicLeapHandTracking/Private/MagicLeapHandTracking.h
+++ b/Engine/Plugins/Lumin/MagicLeap/Source/MagicLeapHandTracking/Private/MagicLeapHandTracking.h
@@ -4,14 +4,9 @@
 
 #include "CoreMinimal.h"
 #include "GenericPlatform/IInputInterface.h"
-#include "IMagicLeapInputDevice.h"
-#include "MagicLeapPluginUtil.h" // for ML_INCLUDES_START/END
-
-#if WITH_MLSDK
-ML_INCLUDES_START
-#include <ml_hand_tracking.h>
-ML_INCLUDES_END
-#endif //WITH_MLSDK
+#include "IInputDevice.h"
+#include "IMagicLeapTrackerEntity.h"
+#include "Lumin/CAPIShims/LuminAPIHandTracking.h"
 
 #include "MagicLeapHandTrackingTypes.h"
 #include "XRMotionControllerBase.h"
@@ -22,7 +17,7 @@ ML_INCLUDES_END
 /**
   * MagicLeap HandTracking
   */
-class FMagicLeapHandTracking : public IMagicLeapInputDevice, public FXRMotionControllerBase, public MagicLeap::IAppEventHandler, public ILiveLinkSource
+class FMagicLeapHandTracking : public IInputDevice, public IMagicLeapTrackerEntity, public FXRMotionControllerBase, public MagicLeap::IAppEventHandler, public ILiveLinkSource
 {
 public:
 	static FName HandCenter_Name;
@@ -146,7 +141,7 @@ public:
 	{
 		FHandState();
 
-		EHandTrackingGesture Gesture = EHandTrackingGesture::NoHand;
+		EMagicLeapHandTrackingGesture Gesture = EMagicLeapHandTrackingGesture::NoHand;
 		float GestureConfidence = 0.0f;
 		FVector HandCenterNormalized = FVector::ZeroVector;
 
@@ -160,12 +155,12 @@ public:
 
 		FTransformRecord HandCenter;
 
-		bool IsValid() const { return Gesture != EHandTrackingGesture::NoHand; }
-		bool GetTransform(EHandTrackingKeypoint KeyPoint, FTransform& OutTransform) const;
-		const FTransformRecord& GetTransform(EHandTrackingKeypoint KeyPoint) const;
+		bool IsValid() const { return Gesture != EMagicLeapHandTrackingGesture::NoHand; }
+		bool GetTransform(EMagicLeapHandTrackingKeypoint KeyPoint, FTransform& OutTransform) const;
+		const FTransformRecord& GetTransform(EMagicLeapHandTrackingKeypoint KeyPoint) const;
 
 	private:
-		FTransformRecord* EnumToTransformMap[EHandTrackingKeypointCount];
+		FTransformRecord* EnumToTransformMap[EMagicLeapHandTrackingKeypointCount];
 	};
 
 public:
@@ -174,6 +169,7 @@ public:
 
 	/** IMotionController interface */
 	virtual bool GetControllerOrientationAndPosition(const int32 ControllerIndex, const FName MotionSource, FRotator& OutOrientation, FVector& OutPosition, float WorldToMetersScale) const override;
+	virtual ETrackingStatus GetControllerTrackingStatus(const int32 ControllerIndex, const FName MotionSource) const override;
 	virtual bool GetControllerOrientationAndPosition(const int32 ControllerIndex, const EControllerHand DeviceHand, FRotator& OutOrientation, FVector& OutPosition, float WorldToMetersScale) const override;
 	virtual ETrackingStatus GetControllerTrackingStatus(const int32 ControllerIndex, const EControllerHand DeviceHand) const override;
 	virtual FName GetMotionControllerDeviceTypeName() const override;
@@ -188,7 +184,7 @@ public:
 	virtual FText GetSourceType() const override;
 	// End ILiveLinkSource
 
-	/** IMagicLeapInputDevice interface */
+	/** IInputDevice interface */
 	virtual void Tick(float DeltaTime) override;
 	virtual void SendControllerEvents() override;
 	virtual void SetMessageHandler(const TSharedRef<FGenericApplicationMessageHandler>& InMessageHandler) override;
@@ -196,22 +192,22 @@ public:
 	virtual void SetChannelValue(int32 ControllerId, FForceFeedbackChannelType ChannelType, float Value) override {};
 	virtual void SetChannelValues(int32 ControllerId, const FForceFeedbackValues &values) override {};
 	virtual bool IsGamepadAttached() const override;
-	virtual void Enable() override;
-	virtual bool SupportsExplicitEnable() const override;
-	virtual void Disable() override;
-	virtual void OnBeginRendering_GameThread_Update() override;
+
+	/** IMagicLeapTrackerEntity interface */
+	virtual void DestroyEntityTracker() override;
+	virtual void OnBeginRendering_GameThread() override;
 
 	const FHandState& GetLeftHandState() const;
 	const FHandState& GetRightHandState() const;
 	bool IsHandTrackingStateValid() const;
 
-	bool GetKeypointTransform(EControllerHand Hand, EHandTrackingKeypoint Keypoint, FTransform& OutTransform) const;
+	bool GetKeypointTransform(EControllerHand Hand, EMagicLeapHandTrackingKeypoint Keypoint, FTransform& OutTransform) const;
 
-	bool SetConfiguration(bool bTrackingEnabled, const TArray<EHandTrackingGesture>& ActiveKeyPoses, EHandTrackingKeypointFilterLevel KeypointsFilterLevel, EHandTrackingGestureFilterLevel GestureFilterLevel);
-	bool GetConfiguration(bool& bTrackingEnabled, TArray<EHandTrackingGesture>& ActiveKeyPoses, EHandTrackingKeypointFilterLevel& KeypointsFilterLevel, EHandTrackingGestureFilterLevel& GestureFilterLevel);
+	bool SetConfiguration(bool bTrackingEnabled, const TArray<EMagicLeapHandTrackingGesture>& ActiveKeyPoses, EMagicLeapHandTrackingKeypointFilterLevel KeypointsFilterLevel, EMagicLeapHandTrackingGestureFilterLevel GestureFilterLevel);
+	bool GetConfiguration(bool& bTrackingEnabled, TArray<EMagicLeapHandTrackingGesture>& ActiveKeyPoses, EMagicLeapHandTrackingKeypointFilterLevel& KeypointsFilterLevel, EMagicLeapHandTrackingGestureFilterLevel& GestureFilterLevel);
 
-	void SetGestureConfidenceThreshold(EHandTrackingGesture Gesture, float Confidence);
-	float GetGestureConfidenceThreshold(EHandTrackingGesture Gesture) const;
+	void SetGestureConfidenceThreshold(EMagicLeapHandTrackingGesture Gesture, float Confidence);
+	float GetGestureConfidenceThreshold(EMagicLeapHandTrackingGesture Gesture) const;
 
 private:
 	void BuildKeypointMaps();
@@ -233,7 +229,7 @@ private:
 	const MLHandTrackingData& GetPreviousHandTrackingData()				{ return HandTrackingDatas[1 - CurrentHandTrackingDataIndex]; }
 	void SendControllerEventsForHand(const MLHandTrackingHandState& NewHandState, const MLHandTrackingHandState& OldHandState, const TArray<FName>& GestureMap);
 
-	static EHandTrackingGesture TranslateGestureEnum(MLHandTrackingKeyPose KeyPose);
+	static EMagicLeapHandTrackingGesture TranslateGestureEnum(MLHandTrackingKeyPose KeyPose);
 #endif //WITH_MLSDK
 
 	TSharedPtr<FGenericApplicationMessageHandler> MessageHandler;
@@ -248,7 +244,7 @@ private:
 	MLHandTrackingStaticData HandTrackingStaticData;
 #endif //WITH_MLSDK
 	TArray<int32> BoneParents;
-	TArray<EHandTrackingKeypoint> BoneKeypoints;
+	TArray<EMagicLeapHandTrackingKeypoint> BoneKeypoints;
 
 	FHandState LeftHand;
 	FHandState RightHand;
@@ -270,6 +266,8 @@ private:
 		static const FKey Left_Pinch;
 		static const FKey Left_Thumb;
 		static const FKey Left_L;
+		static const FKey Left_OpenHand;
+		// DEPRECATED
 		static const FKey Left_OpenHandBack;
 		static const FKey Left_Ok;
 		static const FKey Left_C;
@@ -281,6 +279,8 @@ private:
 		static const FKey Right_Pinch;
 		static const FKey Right_Thumb;
 		static const FKey Right_L;
+		static const FKey Right_OpenHand;
+		// DEPRECATED
 		static const FKey Right_OpenHandBack;
 		static const FKey Right_Ok;
 		static const FKey Right_C;
diff --git a/Engine/Plugins/Lumin/MagicLeap/Source/MagicLeapHandTracking/Private/MagicLeapHandTrackingFunctionLibrary.cpp b/Engine/Plugins/Lumin/MagicLeap/Source/MagicLeapHandTracking/Private/MagicLeapHandTrackingFunctionLibrary.cpp
--- a/Engine/Plugins/Lumin/MagicLeap/Source/MagicLeapHandTracking/Private/MagicLeapHandTrackingFunctionLibrary.cpp
+++ b/Engine/Plugins/Lumin/MagicLeap/Source/MagicLeapHandTracking/Private/MagicLeapHandTrackingFunctionLibrary.cpp
@@ -9,141 +9,137 @@
 #include "ILiveLinkClient.h"
 #include "ILiveLinkSource.h"
 
-bool UMagicLeapHandTrackingFunctionLibrary::GetHandCenter(EControllerHand Hand, FTransform& HandCenter)
+bool GetHandTrackingData(EControllerHand Hand, const FMagicLeapHandTracking::FHandState*& HandTrackingData)
 {
 	TSharedPtr<FMagicLeapHandTracking> HandTracking = StaticCastSharedPtr<FMagicLeapHandTracking>(IMagicLeapHandTrackingPlugin::Get().GetInputDevice());
 
 	if (HandTracking.IsValid() && HandTracking->IsHandTrackingStateValid())
 	{
-		if (Hand == EControllerHand::Left || Hand == EControllerHand::Right)
+		switch (Hand)
 		{
-			const FMagicLeapHandTracking::FHandState& HandTrackingData = (Hand == EControllerHand::Right) ? HandTracking->GetRightHandState() : HandTracking->GetLeftHandState();
-			HandCenter = HandTrackingData.HandCenter.Transform;
-			return HandTrackingData.IsValid();
-		}
-		else
+		case EControllerHand::Left:
 		{
-			UE_LOG(LogMagicLeapHandTracking, Error, TEXT("Hand %d is not supported"), static_cast<int32>(Hand));
-			return false;
+			HandTrackingData = &HandTracking->GetLeftHandState();
 		}
-	}
-
-	return false;
-}
+		break;
 
-bool UMagicLeapHandTrackingFunctionLibrary::GetHandIndexFingerTip(EControllerHand Hand, EGestureTransformSpace TransformSpace, FTransform& Pointer)
-{
-	TSharedPtr<FMagicLeapHandTracking> HandTracking = StaticCastSharedPtr<FMagicLeapHandTracking>(IMagicLeapHandTrackingPlugin::Get().GetInputDevice());
+		case EControllerHand::Right:
+		{
+			HandTrackingData = &HandTracking->GetRightHandState();
+		}
+		break;
 
-	if (HandTracking.IsValid() && HandTracking->IsHandTrackingStateValid())
-	{
-		if (Hand == EControllerHand::Left || Hand == EControllerHand::Right)
+		case EControllerHand::AnyHand:
 		{
-			const FMagicLeapHandTracking::FHandState& HandTrackingData = (Hand == EControllerHand::Right) ? HandTracking->GetRightHandState() : HandTracking->GetLeftHandState();
-			switch (TransformSpace)
+			HandTrackingData = &HandTracking->GetLeftHandState();
+			if (!HandTrackingData->IsValid() ||
+				(HandTracking->GetRightHandState().IsValid() &&
+				(HandTracking->GetRightHandState().GestureConfidence > HandTracking->GetLeftHandState().GestureConfidence)))
 			{
-			case EGestureTransformSpace::Tracking:
-			{
-				Pointer = HandTrackingData.IndexFinger.Tip.Transform;
-				break;
-			}
-			case EGestureTransformSpace::World:
-			{
-				FTransform TrackingToWorldTransform = UHeadMountedDisplayFunctionLibrary::GetTrackingToWorldTransform(GWorld);
-				Pointer = HandTrackingData.IndexFinger.Tip.Transform * TrackingToWorldTransform;
-				break;
-			}
-			case EGestureTransformSpace::Hand:
-			{
-				Pointer = HandTrackingData.IndexFinger.Tip.Transform * HandTrackingData.HandCenter.Transform.Inverse();
-				break;
+				HandTrackingData = &HandTracking->GetRightHandState();
 			}
-			default:
-				check(false);
-				return false;
-			}
-			return HandTrackingData.IsValid();
 		}
-		else
+		break;
+
+		default:
 		{
 			UE_LOG(LogMagicLeapHandTracking, Error, TEXT("Hand %d is not supported"), static_cast<int32>(Hand));
-			return false;
 		}
+		}
+	}
+
+	return HandTrackingData && HandTrackingData->IsValid();
+}
+
+bool UMagicLeapHandTrackingFunctionLibrary::GetHandCenter(EControllerHand Hand, FTransform& HandCenter)
+{
+	const FMagicLeapHandTracking::FHandState* HandTrackingData = nullptr;
+
+	if (GetHandTrackingData(Hand, HandTrackingData))
+	{
+		HandCenter = HandTrackingData->HandCenter.Transform;
+		return true;
 	}
 
 	return false;
 }
 
-bool UMagicLeapHandTrackingFunctionLibrary::GetHandThumbTip(EControllerHand Hand, EGestureTransformSpace TransformSpace, FTransform& Secondary)
+bool UMagicLeapHandTrackingFunctionLibrary::GetHandIndexFingerTip(EControllerHand Hand, EMagicLeapGestureTransformSpace TransformSpace, FTransform& Pointer)
 {
-	TSharedPtr<FMagicLeapHandTracking> HandTracking = StaticCastSharedPtr<FMagicLeapHandTracking>(IMagicLeapHandTrackingPlugin::Get().GetInputDevice());
+	const FMagicLeapHandTracking::FHandState* HandTrackingData = nullptr;
 
-	if (HandTracking.IsValid() && HandTracking->IsHandTrackingStateValid())
+	if (GetHandTrackingData(Hand, HandTrackingData))
 	{
-		if (Hand == EControllerHand::Left || Hand == EControllerHand::Right)
+		switch (TransformSpace)
 		{
-			const FMagicLeapHandTracking::FHandState& HandTrackingData = (Hand == EControllerHand::Right) ? HandTracking->GetRightHandState() : HandTracking->GetLeftHandState();
-			switch (TransformSpace)
-			{
-			case EGestureTransformSpace::Tracking:
-			{
-				Secondary = HandTrackingData.Thumb.Tip.Transform;
-				break;
-			}
-			case EGestureTransformSpace::World:
-			{
-				FTransform TrackingToWorldTransform = UHeadMountedDisplayFunctionLibrary::GetTrackingToWorldTransform(GWorld);
-				Secondary = HandTrackingData.Thumb.Tip.Transform * TrackingToWorldTransform;
-				break;
-			}
-			case EGestureTransformSpace::Hand:
-			{
-				Secondary = HandTrackingData.Thumb.Tip.Transform * HandTrackingData.HandCenter.Transform.Inverse();
-				break;
-			}
-			default:
-				check(false);
-				return false;
-			}
-			return HandTrackingData.IsValid();
+		case EMagicLeapGestureTransformSpace::Tracking:
+		{
+			Pointer = HandTrackingData->IndexFinger.Tip.Transform;
+			break;
 		}
-		else
+		case EMagicLeapGestureTransformSpace::World:
 		{
-			UE_LOG(LogMagicLeapHandTracking, Error, TEXT("Hand %d is not supported"), static_cast<int32>(Hand));
+			const FTransform TrackingToWorldTransform = UHeadMountedDisplayFunctionLibrary::GetTrackingToWorldTransform(nullptr);
+			Pointer = HandTrackingData->IndexFinger.Tip.Transform * TrackingToWorldTransform;
+			break;
+		}
+		case EMagicLeapGestureTransformSpace::Hand:
+		{
+			Pointer = HandTrackingData->IndexFinger.Tip.Transform * HandTrackingData->HandCenter.Transform.Inverse();
+			break;
+		}
+		default:
+			check(false);
 			return false;
 		}
+		return true;
 	}
 
 	return false;
 }
 
-bool UMagicLeapHandTrackingFunctionLibrary::GetHandCenterNormalized(EControllerHand Hand, FVector& HandCenterNormalized)
+bool UMagicLeapHandTrackingFunctionLibrary::GetHandThumbTip(EControllerHand Hand, EMagicLeapGestureTransformSpace TransformSpace, FTransform& Secondary)
 {
-	TSharedPtr<FMagicLeapHandTracking> HandTracking = StaticCastSharedPtr<FMagicLeapHandTracking>(IMagicLeapHandTrackingPlugin::Get().GetInputDevice());
+	const FMagicLeapHandTracking::FHandState* HandTrackingData = nullptr;
 
-	if (HandTracking.IsValid() && HandTracking->IsHandTrackingStateValid())
+	if (GetHandTrackingData(Hand, HandTrackingData))
 	{
-		if (Hand == EControllerHand::Left)
+		switch (TransformSpace)
 		{
-			if (HandTracking->GetLeftHandState().IsValid())
-			{
-				HandCenterNormalized = HandTracking->GetLeftHandState().HandCenterNormalized;
-			}
-			return HandTracking->GetLeftHandState().IsValid();
+		case EMagicLeapGestureTransformSpace::Tracking:
+		{
+			Secondary = HandTrackingData->Thumb.Tip.Transform;
+			break;
 		}
-		else if (Hand == EControllerHand::Right)
+		case EMagicLeapGestureTransformSpace::World:
 		{
-			if (HandTracking->GetRightHandState().IsValid())
-			{
-				HandCenterNormalized = HandTracking->GetRightHandState().HandCenterNormalized;
-			}
-			return HandTracking->GetRightHandState().IsValid();
+			const FTransform TrackingToWorldTransform = UHeadMountedDisplayFunctionLibrary::GetTrackingToWorldTransform(nullptr);
+			Secondary = HandTrackingData->Thumb.Tip.Transform * TrackingToWorldTransform;
+			break;
 		}
-		else
+		case EMagicLeapGestureTransformSpace::Hand:
 		{
-			UE_LOG(LogMagicLeapHandTracking, Error, TEXT("Hand %d is not supported"), static_cast<int32>(Hand));
+			Secondary = HandTrackingData->Thumb.Tip.Transform * HandTrackingData->HandCenter.Transform.Inverse();
+			break;
+		}
+		default:
+			check(false);
 			return false;
 		}
+		return true;
+	}
+
+	return false;
+}
+
+bool UMagicLeapHandTrackingFunctionLibrary::GetHandCenterNormalized(EControllerHand Hand, FVector& HandCenterNormalized)
+{
+	const FMagicLeapHandTracking::FHandState* HandTrackingData = nullptr;
+
+	if (GetHandTrackingData(Hand, HandTrackingData))
+	{
+		HandCenterNormalized = HandTrackingData->HandCenterNormalized;
+		return true;
 	}
 
 	return false;
@@ -151,30 +147,21 @@ bool UMagicLeapHandTrackingFunctionLibrary::GetHandCenterNormalized(EControllerH
 
 bool UMagicLeapHandTrackingFunctionLibrary::GetGestureKeypoints(EControllerHand Hand, TArray<FTransform>& Keypoints)
 {
-	TSharedPtr<FMagicLeapHandTracking> HandTracking = StaticCastSharedPtr<FMagicLeapHandTracking>(IMagicLeapHandTrackingPlugin::Get().GetInputDevice());
+	const FMagicLeapHandTracking::FHandState* HandTrackingData = nullptr;
 
-	if (HandTracking.IsValid() && HandTracking->IsHandTrackingStateValid())
+	if (GetHandTrackingData(Hand, HandTrackingData))
 	{
-		if (Hand == EControllerHand::Left || Hand == EControllerHand::Right)
-		{
-			const FMagicLeapHandTracking::FHandState& HandState = (Hand == EControllerHand::Left) ? HandTracking->GetLeftHandState() : HandTracking->GetRightHandState();
-			Keypoints.SetNum(3);
-			Keypoints[0] = HandState.HandCenter.Transform;
-			Keypoints[1] = HandState.IndexFinger.Tip.Transform;
-			Keypoints[2] = HandState.Thumb.Tip.Transform;
-			return true;
-		}
-		else
-		{
-			UE_LOG(LogMagicLeapHandTracking, Error, TEXT("Hand %d is not supported"), static_cast<int32>(Hand));
-			return false;
-		}
+		Keypoints.SetNum(3);
+		Keypoints[0] = HandTrackingData->HandCenter.Transform;
+		Keypoints[1] = HandTrackingData->IndexFinger.Tip.Transform;
+		Keypoints[2] = HandTrackingData->Thumb.Tip.Transform;
+		return true;
 	}
 
 	return false;
 }
 
-bool UMagicLeapHandTrackingFunctionLibrary::GetGestureKeypointTransform(EControllerHand Hand, EHandTrackingKeypoint Keypoint, EGestureTransformSpace TransformSpace, FTransform& OutTransform)
+bool UMagicLeapHandTrackingFunctionLibrary::GetGestureKeypointTransform(EControllerHand Hand, EMagicLeapHandTrackingKeypoint Keypoint, EMagicLeapGestureTransformSpace TransformSpace, FTransform& OutTransform)
 {
 	TSharedPtr<FMagicLeapHandTracking> HandTracking = StaticCastSharedPtr<FMagicLeapHandTracking>(IMagicLeapHandTrackingPlugin::Get().GetInputDevice());
 
@@ -187,21 +174,21 @@ bool UMagicLeapHandTrackingFunctionLibrary::GetGestureKeypointTransform(EControl
 		{
 			switch (TransformSpace)
 			{
-			case EGestureTransformSpace::Tracking:
+			case EMagicLeapGestureTransformSpace::Tracking:
 			{
 				OutTransform = KeyPointTransform;
 				return true;
 			}
-			case EGestureTransformSpace::World:
+			case EMagicLeapGestureTransformSpace::World:
 			{
-				FTransform TrackingToWorldTransform = UHeadMountedDisplayFunctionLibrary::GetTrackingToWorldTransform(GWorld);
+				const FTransform TrackingToWorldTransform = UHeadMountedDisplayFunctionLibrary::GetTrackingToWorldTransform(nullptr);
 				OutTransform = KeyPointTransform * TrackingToWorldTransform;
 				return true;
 			}
-			case EGestureTransformSpace::Hand:
+			case EMagicLeapGestureTransformSpace::Hand:
 			{
 				FTransform HandTransform;
-				const bool bSuccess2 = HandTracking->GetKeypointTransform(Hand, EHandTrackingKeypoint::Hand_Center, HandTransform);
+				const bool bSuccess2 = HandTracking->GetKeypointTransform(Hand, EMagicLeapHandTrackingKeypoint::Hand_Center, HandTransform);
 				if (bSuccess2)
 				{
 					OutTransform = KeyPointTransform * HandTransform.Inverse();
@@ -219,19 +206,19 @@ bool UMagicLeapHandTrackingFunctionLibrary::GetGestureKeypointTransform(EControl
 	return false;
 }
 
-bool UMagicLeapHandTrackingFunctionLibrary::SetConfiguration(const TArray<EHandTrackingGesture>& StaticGesturesToActivate, EHandTrackingKeypointFilterLevel KeypointsFilterLevel, EHandTrackingGestureFilterLevel GestureFilterLevel, EHandTrackingGestureFilterLevel HandSwitchingFilterLevel, bool bEnabled)
+bool UMagicLeapHandTrackingFunctionLibrary::SetConfiguration(const TArray<EMagicLeapHandTrackingGesture>& StaticGesturesToActivate, EMagicLeapHandTrackingKeypointFilterLevel KeypointsFilterLevel, EMagicLeapHandTrackingGestureFilterLevel GestureFilterLevel, bool bEnabled)
 {
 	TSharedPtr<FMagicLeapHandTracking> HandTracking = StaticCastSharedPtr<FMagicLeapHandTracking>(IMagicLeapHandTrackingPlugin::Get().GetInputDevice());
 	return HandTracking.IsValid() && HandTracking->SetConfiguration(bEnabled, StaticGesturesToActivate, KeypointsFilterLevel, GestureFilterLevel);
 }
 
-bool UMagicLeapHandTrackingFunctionLibrary::GetConfiguration(TArray<EHandTrackingGesture>& ActiveStaticGestures, EHandTrackingKeypointFilterLevel& KeypointsFilterLevel, EHandTrackingGestureFilterLevel& GestureFilterLevel, EHandTrackingGestureFilterLevel& HandSwitchingFilterLevel, bool& bEnabled)
+bool UMagicLeapHandTrackingFunctionLibrary::GetConfiguration(TArray<EMagicLeapHandTrackingGesture>& ActiveStaticGestures, EMagicLeapHandTrackingKeypointFilterLevel& KeypointsFilterLevel, EMagicLeapHandTrackingGestureFilterLevel& GestureFilterLevel, bool& bEnabled)
 {
 	TSharedPtr<FMagicLeapHandTracking> HandTracking = StaticCastSharedPtr<FMagicLeapHandTracking>(IMagicLeapHandTrackingPlugin::Get().GetInputDevice());
 	return HandTracking.IsValid() && HandTracking->GetConfiguration(bEnabled, ActiveStaticGestures, KeypointsFilterLevel, GestureFilterLevel);
 }
 
-void UMagicLeapHandTrackingFunctionLibrary::SetStaticGestureConfidenceThreshold(EHandTrackingGesture Gesture, float Confidence)
+void UMagicLeapHandTrackingFunctionLibrary::SetStaticGestureConfidenceThreshold(EMagicLeapHandTrackingGesture Gesture, float Confidence)
 {
 	TSharedPtr<FMagicLeapHandTracking> HandTracking = StaticCastSharedPtr<FMagicLeapHandTracking>(IMagicLeapHandTrackingPlugin::Get().GetInputDevice());
 	if (HandTracking.IsValid())
@@ -240,7 +227,7 @@ void UMagicLeapHandTrackingFunctionLibrary::SetStaticGestureConfidenceThreshold(
 	}
 }
 
-float UMagicLeapHandTrackingFunctionLibrary::GetStaticGestureConfidenceThreshold(EHandTrackingGesture Gesture)
+float UMagicLeapHandTrackingFunctionLibrary::GetStaticGestureConfidenceThreshold(EMagicLeapHandTrackingGesture Gesture)
 {
 	TSharedPtr<FMagicLeapHandTracking> HandTracking = StaticCastSharedPtr<FMagicLeapHandTracking>(IMagicLeapHandTrackingPlugin::Get().GetInputDevice());
 	return (HandTracking.IsValid()) ? HandTracking->GetGestureConfidenceThreshold(Gesture) : 0.0f;
@@ -248,55 +235,28 @@ float UMagicLeapHandTrackingFunctionLibrary::GetStaticGestureConfidenceThreshold
 
 bool UMagicLeapHandTrackingFunctionLibrary::GetCurrentGestureConfidence(EControllerHand Hand, float& Confidence)
 {
-	TSharedPtr<FMagicLeapHandTracking> HandTracking = StaticCastSharedPtr<FMagicLeapHandTracking>(IMagicLeapHandTrackingPlugin::Get().GetInputDevice());
+	const FMagicLeapHandTracking::FHandState* HandTrackingData = nullptr;
 
-	if (HandTracking.IsValid() && HandTracking->IsHandTrackingStateValid())
+	if (GetHandTrackingData(Hand, HandTrackingData))
 	{
-		if (Hand == EControllerHand::Left)
-		{
-			Confidence = HandTracking->GetLeftHandState().GestureConfidence;
-			return true;
-		}
-		else if (Hand == EControllerHand::Right)
-		{
-			Confidence = HandTracking->GetRightHandState().GestureConfidence;
-			return true;
-		}
-		else
-		{
-			UE_LOG(LogMagicLeapHandTracking, Error, TEXT("Hand %d is not supported"), static_cast<int32>(Hand));
-			return false;
-		}
+		Confidence = HandTrackingData->GestureConfidence;
+		return true;
 	}
 
 	return false;
 }
 
-bool UMagicLeapHandTrackingFunctionLibrary::GetCurrentGesture(EControllerHand Hand, EHandTrackingGesture& Gesture)
+bool UMagicLeapHandTrackingFunctionLibrary::GetCurrentGesture(EControllerHand Hand, EMagicLeapHandTrackingGesture& Gesture)
 {
-	TSharedPtr<FMagicLeapHandTracking> HandTracking = StaticCastSharedPtr<FMagicLeapHandTracking>(IMagicLeapHandTrackingPlugin::Get().GetInputDevice());
+	const FMagicLeapHandTracking::FHandState* HandTrackingData = nullptr;
 
-	if (HandTracking.IsValid() && HandTracking->IsHandTrackingStateValid())
+	if (GetHandTrackingData(Hand, HandTrackingData))
 	{
-		if (Hand == EControllerHand::Left)
-		{
-			Gesture = HandTracking->GetLeftHandState().Gesture;
-			return true;
-		}
-		else if (Hand == EControllerHand::Right)
-		{
-			Gesture = HandTracking->GetRightHandState().Gesture;
-			return true;
-		}
-		else
-		{
-			UE_LOG(LogMagicLeapHandTracking, Error, TEXT("Hand %d is not supported"), static_cast<int32>(Hand));
-			Gesture = EHandTrackingGesture::NoHand;
-			return false;
-		}
+		Gesture = HandTrackingData->Gesture;
+		return true;
 	}
 
-	Gesture = EHandTrackingGesture::NoHand;
+	Gesture = EMagicLeapHandTrackingGesture::NoHand;
 	return false;
 }
 
diff --git a/Engine/Plugins/Lumin/MagicLeap/Source/MagicLeapHandTracking/Private/MagicLeapHandTrackingLiveLink.cpp b/Engine/Plugins/Lumin/MagicLeap/Source/MagicLeapHandTracking/Private/MagicLeapHandTrackingLiveLink.cpp
--- a/Engine/Plugins/Lumin/MagicLeap/Source/MagicLeapHandTracking/Private/MagicLeapHandTrackingLiveLink.cpp
+++ b/Engine/Plugins/Lumin/MagicLeap/Source/MagicLeapHandTracking/Private/MagicLeapHandTrackingLiveLink.cpp
@@ -50,45 +50,45 @@ FText FMagicLeapHandTracking::GetSourceType() const
 #define MLHTBONE(keyenum, name, parent) BoneKeypoints.Add(keyenum); StaticData.BoneNames.Add(name); StaticData.BoneParents.Add(parent);
 void FMagicLeapHandTracking::SetupLiveLinkData(FLiveLinkSkeletonStaticData& StaticData)
 {
-	StaticData.BoneNames.Reserve(EHandTrackingKeypointCount);
+	StaticData.BoneNames.Reserve(EMagicLeapHandTrackingKeypointCount);
 
 	// Array of bone indices to parent bone index
-	StaticData.BoneParents.Reserve(EHandTrackingKeypointCount);
+	StaticData.BoneParents.Reserve(EMagicLeapHandTrackingKeypointCount);
 
-	BoneKeypoints.Reserve(EHandTrackingKeypointCount);
+	BoneKeypoints.Reserve(EMagicLeapHandTrackingKeypointCount);
 
 	// Only one can be the root (-1 parent)!
 
-	MLHTBONE(EHandTrackingKeypoint::Thumb_Tip, ThumbTip_Name,			1);
-	MLHTBONE(EHandTrackingKeypoint::Thumb_IP,	ThumbIP_Name,			2);
-	MLHTBONE(EHandTrackingKeypoint::Thumb_MCP, ThumbMCP_Name,			3);
-	MLHTBONE(EHandTrackingKeypoint::Thumb_CMC, ThumbCMC_Name,			23); //3
+	MLHTBONE(EMagicLeapHandTrackingKeypoint::Thumb_Tip, ThumbTip_Name,			1);
+	MLHTBONE(EMagicLeapHandTrackingKeypoint::Thumb_IP,	ThumbIP_Name,			2);
+	MLHTBONE(EMagicLeapHandTrackingKeypoint::Thumb_MCP, ThumbMCP_Name,			3);
+	MLHTBONE(EMagicLeapHandTrackingKeypoint::Thumb_CMC, ThumbCMC_Name,			23); //3
 
-	MLHTBONE(EHandTrackingKeypoint::Index_Tip, IndexFingerTip_Name,		5);
-	MLHTBONE(EHandTrackingKeypoint::Index_DIP, IndexFingerDIP_Name,		6);
-	MLHTBONE(EHandTrackingKeypoint::Index_PIP, IndexFingerPIP_Name,		7);
-	MLHTBONE(EHandTrackingKeypoint::Index_MCP, IndexFingerMCP_Name,		23); //7
+	MLHTBONE(EMagicLeapHandTrackingKeypoint::Index_Tip, IndexFingerTip_Name,		5);
+	MLHTBONE(EMagicLeapHandTrackingKeypoint::Index_DIP, IndexFingerDIP_Name,		6);
+	MLHTBONE(EMagicLeapHandTrackingKeypoint::Index_PIP, IndexFingerPIP_Name,		7);
+	MLHTBONE(EMagicLeapHandTrackingKeypoint::Index_MCP, IndexFingerMCP_Name,		23); //7
 
-	MLHTBONE(EHandTrackingKeypoint::Middle_Tip, MiddleFingerTip_Name,	9);
-	MLHTBONE(EHandTrackingKeypoint::Middle_DIP, MiddleFingerDIP_Name,	10);
-	MLHTBONE(EHandTrackingKeypoint::Middle_PIP, MiddleFingerPIP_Name,	11);
-	MLHTBONE(EHandTrackingKeypoint::Middle_MCP, MiddleFingerMCP_Name,	23); //11
+	MLHTBONE(EMagicLeapHandTrackingKeypoint::Middle_Tip, MiddleFingerTip_Name,	9);
+	MLHTBONE(EMagicLeapHandTrackingKeypoint::Middle_DIP, MiddleFingerDIP_Name,	10);
+	MLHTBONE(EMagicLeapHandTrackingKeypoint::Middle_PIP, MiddleFingerPIP_Name,	11);
+	MLHTBONE(EMagicLeapHandTrackingKeypoint::Middle_MCP, MiddleFingerMCP_Name,	23); //11
 
-	MLHTBONE(EHandTrackingKeypoint::Ring_Tip, RingFingerTip_Name,		13);
-	MLHTBONE(EHandTrackingKeypoint::Ring_DIP, RingFingerDIP_Name,		14);
-	MLHTBONE(EHandTrackingKeypoint::Ring_PIP, RingFingerPIP_Name,		15);
-	MLHTBONE(EHandTrackingKeypoint::Ring_MCP, RingFingerMCP_Name,		23); //15
+	MLHTBONE(EMagicLeapHandTrackingKeypoint::Ring_Tip, RingFingerTip_Name,		13);
+	MLHTBONE(EMagicLeapHandTrackingKeypoint::Ring_DIP, RingFingerDIP_Name,		14);
+	MLHTBONE(EMagicLeapHandTrackingKeypoint::Ring_PIP, RingFingerPIP_Name,		15);
+	MLHTBONE(EMagicLeapHandTrackingKeypoint::Ring_MCP, RingFingerMCP_Name,		23); //15
 
-	MLHTBONE(EHandTrackingKeypoint::Pinky_Tip, PinkyFingerTip_Name,		17);
-	MLHTBONE(EHandTrackingKeypoint::Pinky_DIP, PinkyFingerDIP_Name,		18);
-	MLHTBONE(EHandTrackingKeypoint::Pinky_PIP, PinkyFingerPIP_Name,		19);
-	MLHTBONE(EHandTrackingKeypoint::Pinky_MCP, PinkyFingerMCP_Name,		23); //19
+	MLHTBONE(EMagicLeapHandTrackingKeypoint::Pinky_Tip, PinkyFingerTip_Name,		17);
+	MLHTBONE(EMagicLeapHandTrackingKeypoint::Pinky_DIP, PinkyFingerDIP_Name,		18);
+	MLHTBONE(EMagicLeapHandTrackingKeypoint::Pinky_PIP, PinkyFingerPIP_Name,		19);
+	MLHTBONE(EMagicLeapHandTrackingKeypoint::Pinky_MCP, PinkyFingerMCP_Name,		23); //19
 
-	MLHTBONE(EHandTrackingKeypoint::Wrist_Center,	WristCenter_Name,	23);
-	MLHTBONE(EHandTrackingKeypoint::Wrist_Ulnar,	WristUlnar_Name,	23);
-	MLHTBONE(EHandTrackingKeypoint::Wrist_Radial,	WristRadial_Name,	23);
+	MLHTBONE(EMagicLeapHandTrackingKeypoint::Wrist_Center,	WristCenter_Name,	23);
+	MLHTBONE(EMagicLeapHandTrackingKeypoint::Wrist_Ulnar,	WristUlnar_Name,	23);
+	MLHTBONE(EMagicLeapHandTrackingKeypoint::Wrist_Radial,	WristRadial_Name,	23);
 
-	MLHTBONE(EHandTrackingKeypoint::Hand_Center, HandCenter_Name,		-1); //23 //Root
+	MLHTBONE(EMagicLeapHandTrackingKeypoint::Hand_Center, HandCenter_Name,		-1); //23 //Root
 }
 #undef MLHTBONE
 
@@ -99,12 +99,12 @@ void FMagicLeapHandTracking::UpdateLiveLinkTransforms(TArray<FTransform>& OutTra
 	// Sadly hand tracking transforms can be unused, and contain only an identity matrix.  So we do some weird stuff to keep those
 	// identity in the hierarchical skeleton.
 
-	if (EHandTrackingKeypointCount != BoneKeypoints.Num())
+	if (EMagicLeapHandTrackingKeypointCount != BoneKeypoints.Num())
 	{
 		return;
 	}
 
-	for (int32 i = 0; i < EHandTrackingKeypointCount; ++i)
+	for (int32 i = 0; i < EMagicLeapHandTrackingKeypointCount; ++i)
 	{
 		const FTransformRecord& TransformRecord = HandState.GetTransform(BoneKeypoints[i]);
 		bool bWritten = TransformRecord.bWritten;
@@ -173,9 +173,9 @@ void FMagicLeapHandTracking::UpdateLiveLink()
 				LiveLinkRightFrame.InitializeWith(FLiveLinkAnimationFrameData::StaticStruct(), nullptr);
 				FLiveLinkAnimationFrameData* LiveLinkLeftFramePtr = LiveLinkLeftFrame.Cast<FLiveLinkAnimationFrameData>();
 				FLiveLinkAnimationFrameData* LiveLinkRightFramePtr = LiveLinkRightFrame.Cast<FLiveLinkAnimationFrameData>();
-				LiveLinkLeftFramePtr->Transforms.Reserve(EHandTrackingKeypointCount);
-				LiveLinkRightFramePtr->Transforms.Reserve(EHandTrackingKeypointCount);
-				for (size_t i = 0; i < EHandTrackingKeypointCount; ++i)
+				LiveLinkLeftFramePtr->Transforms.Reserve(EMagicLeapHandTrackingKeypointCount);
+				LiveLinkRightFramePtr->Transforms.Reserve(EMagicLeapHandTrackingKeypointCount);
+				for (size_t i = 0; i < EMagicLeapHandTrackingKeypointCount; ++i)
 				{
 					LiveLinkRightFramePtr->Transforms.Add(FTransform::Identity);
 					LiveLinkRightFramePtr->Transforms.Add(FTransform::Identity);
diff --git a/Engine/Plugins/Lumin/MagicLeap/Source/MagicLeapHandTracking/Public/MagicLeapHandTrackingFunctionLibrary.h b/Engine/Plugins/Lumin/MagicLeap/Source/MagicLeapHandTracking/Public/MagicLeapHandTrackingFunctionLibrary.h
--- a/Engine/Plugins/Lumin/MagicLeap/Source/MagicLeapHandTracking/Public/MagicLeapHandTrackingFunctionLibrary.h
+++ b/Engine/Plugins/Lumin/MagicLeap/Source/MagicLeapHandTracking/Public/MagicLeapHandTrackingFunctionLibrary.h
@@ -26,7 +26,7 @@ public:
 	  @param HandCenter Output parameter containing the position and orientation of the given hand.
 	  @return true if the output param was populated with a valid value, false means that is is either unchanged or populated with a stale value.
 	*/
-	UFUNCTION(BlueprintCallable, Category = "HandTracking|MagicLeap", meta = (DeprecatedFunction, DeprecationMessage = "GetHandCenter is deprecated. Please use GetGestureKeypointTransform instead."))
+	UFUNCTION(BlueprintCallable, Category = "HandTracking|MagicLeap")
 	static bool GetHandCenter(EControllerHand Hand, FTransform& HandCenter);
 
 	/**
@@ -37,8 +37,8 @@ public:
 	  @param Pointer Output parameter containing the position and orientation.
 	  @return true if the output param was populated with a valid value, false means that is is either unchanged or populated with a stale value.
 	  */
-	UFUNCTION(BlueprintCallable, Category = "HandTracking|MagicLeap", meta = (DeprecatedFunction, DeprecationMessage = "GetHandIndexFingerTip is deprecated. Please use GetGestureKeypointTransform instead."))
-	static bool GetHandIndexFingerTip(EControllerHand Hand, EGestureTransformSpace TransformSpace, FTransform& Pointer);
+	UFUNCTION(BlueprintCallable, Category = "HandTracking|MagicLeap")
+	static bool GetHandIndexFingerTip(EControllerHand Hand, EMagicLeapGestureTransformSpace TransformSpace, FTransform& Pointer);
 
 	/**
 	  Transform of Thumb Tip.
@@ -52,8 +52,8 @@ public:
 	  @param Secondary Output parameter containing the position and orientation.
 	  @return true if the output param was populated with a valid value, false means that is is either unchanged or populated with a stale value.
 	  */
-	UFUNCTION(BlueprintCallable, Category = "HandTracking|MagicLeap", meta = (DeprecatedFunction, DeprecationMessage = "GetHandThumbTip is deprecated. Please use GetGestureKeypointTransform instead."))
-	static bool GetHandThumbTip(EControllerHand Hand, EGestureTransformSpace TransformSpace, FTransform& Secondary);
+	UFUNCTION(BlueprintCallable, Category = "HandTracking|MagicLeap")
+	static bool GetHandThumbTip(EControllerHand Hand, EMagicLeapGestureTransformSpace TransformSpace, FTransform& Secondary);
 
 	/**
 	  Normalized position of the center of the given hand. This can be used to detect and warn the users that the hand is out of the gesture detection frame.
@@ -79,7 +79,7 @@ public:
 	  @param Keypoints Output parameter containing transforms of the keypoints detected on the given hand.
 	  @return true if the output param was populated with a valid value, false otherwise.
 	*/
-	UFUNCTION(BlueprintCallable, Category = "HandTracking|MagicLeap", meta = (DeprecatedFunction, DeprecationMessage = "GetGestureKeypoints is deprecated. Please use GetGestureKeypointTransform instead."))
+	UFUNCTION(BlueprintCallable, Category = "HandTracking|MagicLeap")
 	static bool GetGestureKeypoints(EControllerHand Hand, TArray<FTransform>& Keypoints);
 
 	/**
@@ -92,7 +92,7 @@ public:
 	@return true if the output param was populated with a valid value, false means that is is either unchanged or populated with a stale value.
 	*/
 	UFUNCTION(BlueprintCallable, Category = "HandTracking|MagicLeap")
-	static bool GetGestureKeypointTransform(EControllerHand Hand, EHandTrackingKeypoint Keypoint, EGestureTransformSpace TransformSpace, FTransform& Transform);
+	static bool GetGestureKeypointTransform(EControllerHand Hand, EMagicLeapHandTrackingKeypoint Keypoint, EMagicLeapGestureTransformSpace TransformSpace, FTransform& Transform);
 
 	/**
 	  Enables and disables the gestures to be detected by the gesture recognition system.
@@ -104,7 +104,7 @@ public:
 	  @return true if the configuration was set successfully.
 	*/
 	UFUNCTION(BlueprintCallable, Category = "HandTracking|MagicLeap")
-	static bool SetConfiguration(const TArray<EHandTrackingGesture>& StaticGesturesToActivate, EHandTrackingKeypointFilterLevel KeypointsFilterLevel = EHandTrackingKeypointFilterLevel::NoFilter, EHandTrackingGestureFilterLevel GestureFilterLevel = EHandTrackingGestureFilterLevel::NoFilter, EHandTrackingGestureFilterLevel HandSwitchingFilterLevel = EHandTrackingGestureFilterLevel::NoFilter, bool bTrackingEnabled = true);
+	static bool SetConfiguration(const TArray<EMagicLeapHandTrackingGesture>& StaticGesturesToActivate, EMagicLeapHandTrackingKeypointFilterLevel KeypointsFilterLevel = EMagicLeapHandTrackingKeypointFilterLevel::NoFilter, EMagicLeapHandTrackingGestureFilterLevel GestureFilterLevel = EMagicLeapHandTrackingGestureFilterLevel::NoFilter, bool bTrackingEnabled = true);
 
 	/**
 	  Gets the list of static and dynamic gestures currently set to be identified by the gesture recognition system.
@@ -116,7 +116,7 @@ public:
 	  @return true if the output params were populated with a valid value, false otherwise.
 	*/
 	UFUNCTION(BlueprintCallable, Category = "HandTracking|MagicLeap")
-	static bool GetConfiguration(TArray<EHandTrackingGesture>& ActiveStaticGestures, EHandTrackingKeypointFilterLevel& KeypointsFilterLevel, EHandTrackingGestureFilterLevel& GestureFilterLevel, EHandTrackingGestureFilterLevel& HandSwitchingFilterLevel, bool& bTrackingEnabled);
+	static bool GetConfiguration(TArray<EMagicLeapHandTrackingGesture>& ActiveStaticGestures, EMagicLeapHandTrackingKeypointFilterLevel& KeypointsFilterLevel, EMagicLeapHandTrackingGestureFilterLevel& GestureFilterLevel, bool& bTrackingEnabled);
 
 	/**
 	  Sets the minimum gesture confidence to filter out the detected static gesture.
@@ -125,7 +125,7 @@ public:
 	  @param Confidence The gesture confidence threshold.
 	*/
 	UFUNCTION(BlueprintCallable, Category = "HandTracking|MagicLeap")
-	static void SetStaticGestureConfidenceThreshold(EHandTrackingGesture Gesture, float Confidence);
+	static void SetStaticGestureConfidenceThreshold(EMagicLeapHandTrackingGesture Gesture, float Confidence);
 
 	/**
 	  Gets the minimum gesture confidence used to filter out the detected static gesture.
@@ -134,7 +134,7 @@ public:
 	  @return The gesture confidence threshold.
 	*/
 	UFUNCTION(BlueprintCallable, Category = "HandTracking|MagicLeap")
-	static float GetStaticGestureConfidenceThreshold(EHandTrackingGesture Gesture);
+	static float GetStaticGestureConfidenceThreshold(EMagicLeapHandTrackingGesture Gesture);
 
 	/**
 	  The confidence level of the current gesture being performed by the given hand.
@@ -155,7 +155,7 @@ public:
 	  @return true if the output param was populated with a valid value, false otherwise.
 	*/
 	UFUNCTION(BlueprintCallable, Category = "HandTracking|MagicLeap")
-	static bool GetCurrentGesture(EControllerHand Hand, EHandTrackingGesture& Gesture);
+	static bool GetCurrentGesture(EControllerHand Hand, EMagicLeapHandTrackingGesture& Gesture);
 
 
 	/**
diff --git a/Engine/Plugins/Lumin/MagicLeap/Source/MagicLeapHandTracking/Public/MagicLeapHandTrackingTypes.h b/Engine/Plugins/Lumin/MagicLeap/Source/MagicLeapHandTracking/Public/MagicLeapHandTrackingTypes.h
--- a/Engine/Plugins/Lumin/MagicLeap/Source/MagicLeapHandTracking/Public/MagicLeapHandTrackingTypes.h
+++ b/Engine/Plugins/Lumin/MagicLeap/Source/MagicLeapHandTracking/Public/MagicLeapHandTrackingTypes.h
@@ -8,45 +8,45 @@
 
 /*! Transforms that could be tracked on the hand. In 0.15.0 RC5 8 of them may be tracked.*/
 UENUM(BlueprintType)
-enum class EHandTrackingKeypoint : uint8
+enum class EMagicLeapHandTrackingKeypoint : uint8
 {
-	Thumb_Tip, // SDK 0.15.0 tracked
-	Thumb_IP,
-	Thumb_MCP, // SDK 0.15.0 tracked
+	Thumb_Tip, // SDK 0.21.0 tracked
+	Thumb_IP,  // SDK 0.21.0 tracked
+	Thumb_MCP, // SDK 0.21.0 tracked
 	Thumb_CMC,
 
-	Index_Tip, // SDK 0.15.0 tracked
+	Index_Tip, // SDK 0.21.0 tracked
 	Index_DIP,
-	Index_PIP, // SDK 0.15.0 tracked
-	Index_MCP, // SDK 0.15.0 tracked
+	Index_PIP, // SDK 0.21.0 tracked
+	Index_MCP, // SDK 0.21.0 tracked
 
-	Middle_Tip,
+	Middle_Tip, // SDK 0.21.0 tracked
 	Middle_DIP,
-	Middle_PIP,
-	Middle_MCP, // SDK 0.15.0 tracked
+	Middle_PIP, // SDK 0.21.0 tracked
+	Middle_MCP, // SDK 0.21.0 tracked
 
-	Ring_Tip,
+	Ring_Tip, // SDK 0.21.0 tracked
 	Ring_DIP,
 	Ring_PIP,
-	Ring_MCP,
+	Ring_MCP, // SDK 0.21.0 tracked
 
-	Pinky_Tip,
+	Pinky_Tip, // SDK 0.21.0 tracked
 	Pinky_DIP,
 	Pinky_PIP,
-	Pinky_MCP,
+	Pinky_MCP, // SDK 0.21.0 tracked
 
-	Wrist_Center, // SDK 0.15.0 tracked
+	Wrist_Center, // SDK 0.21.0 tracked
 	Wrist_Ulnar,
 	Wrist_Radial,
 
-	Hand_Center // SDK 0.15.0 tracked
+	Hand_Center // SDK 0.21.0 tracked
 };
 
-const int32 EHandTrackingKeypointCount = static_cast<int32>(EHandTrackingKeypoint::Hand_Center) + 1;
+const int32 EMagicLeapHandTrackingKeypointCount = static_cast<int32>(EMagicLeapHandTrackingKeypoint::Hand_Center) + 1;
 
 /*! Static key pose types which are available when both hands are separated. */
 UENUM(BlueprintType)
-enum class EHandTrackingGesture : uint8
+enum class EMagicLeapHandTrackingGesture : uint8
 {
 	/** One finger. */
 	Finger,
@@ -58,8 +58,10 @@ enum class EHandTrackingGesture : uint8
 	Thumb,
 	/** An L shape. */
 	L,
-	/** An open hand with the back of the hand facing the user. */
-	OpenHandBack,
+	/** An open palm of the hand facing the user. */
+	OpenHand,
+	/** DEPRECATED (USE OpenHand INSTEAD) - A back open hand of the hand facing away from the user. */
+	OpenHandBack = EMagicLeapHandTrackingGesture::OpenHand,
 	/** A pinch with all fingers, except the index finger and the thumb, extended out. */
 	Ok,
 	/** A rounded 'C' alphabet shape. */
@@ -72,7 +74,7 @@ enum class EHandTrackingGesture : uint8
 
 /** Filtering for the keypoints and hand centers. */
 UENUM(BlueprintType)
-enum class EHandTrackingKeypointFilterLevel : uint8
+enum class EMagicLeapHandTrackingKeypointFilterLevel : uint8
 {
 	/** No filtering is done, the points are raw. */
 	NoFilter,
@@ -84,7 +86,7 @@ enum class EHandTrackingKeypointFilterLevel : uint8
 
 /** Filtering for the gesture recognition and hand switching. */
 UENUM(BlueprintType)
-enum class EHandTrackingGestureFilterLevel : uint8
+enum class EMagicLeapHandTrackingGestureFilterLevel : uint8
 {
 	/** No filtering is done, the gestures are raw. */
 	NoFilter,
@@ -96,7 +98,7 @@ enum class EHandTrackingGestureFilterLevel : uint8
 
 /** Gesture key point transform spaces. */
 UENUM(BlueprintType)
-enum class EGestureTransformSpace : uint8
+enum class EMagicLeapGestureTransformSpace : uint8
 {
 	/** Gesture key point transforms are reported in Unreal world space. This is more costly on CPU.*/
 	World,
@@ -114,6 +116,8 @@ struct FMagicLeapGestureKeyNames
 	static const FName Left_Pinch_Name;
 	static const FName Left_Thumb_Name;
 	static const FName Left_L_Name;
+	static const FName Left_OpenHand_Name;
+	// DEPRECATED
 	static const FName Left_OpenHandBack_Name;
 	static const FName Left_Ok_Name;
 	static const FName Left_C_Name;
@@ -125,6 +129,8 @@ struct FMagicLeapGestureKeyNames
 	static const FName Right_Pinch_Name;
 	static const FName Right_Thumb_Name;
 	static const FName Right_L_Name;
+	static const FName Right_OpenHand_Name;
+	// DEPRECATED
 	static const FName Right_OpenHandBack_Name;
 	static const FName Right_Ok_Name;
 	static const FName Right_C_Name;
diff --git a/Engine/Plugins/Lumin/MagicLeap/Source/MagicLeapHelperOpenGL/Private/ExternalTextureRendererGL.cpp b/Engine/Plugins/Lumin/MagicLeap/Source/MagicLeapHelperOpenGL/Private/ExternalTextureRendererGL.cpp
new file mode 100644
--- /dev/null
+++ b/Engine/Plugins/Lumin/MagicLeap/Source/MagicLeapHelperOpenGL/Private/ExternalTextureRendererGL.cpp
@@ -0,0 +1,81 @@
+// Copyright 1998-2019 Epic Games, Inc. All Rights Reserved.
+
+#include "ExternalTextureRendererGL.h"
+#include "IMagicLeapHelperOpenGLPlugin.h"
+
+#if PLATFORM_LUMIN
+#include "OpenGLDrv.h"
+#endif // PLATFORM_LUMIN
+
+#if PLATFORM_LUMIN
+class FMagicLeapExternalTextureRendererGLImpl
+{
+public:
+	FMagicLeapExternalTextureRendererGLImpl()
+	: Image(EGL_NO_IMAGE_KHR)
+	{}
+
+	EGLImageKHR Image;
+};
+#endif // PLATFORM_LUMIN
+
+FExternalTextureRendererGL::FExternalTextureRendererGL()
+#if PLATFORM_LUMIN
+: Impl(new FMagicLeapExternalTextureRendererGLImpl())
+#endif // PLATFORM_LUMIN
+{
+
+}
+
+FExternalTextureRendererGL::~FExternalTextureRendererGL()
+{
+#if PLATFORM_LUMIN
+	delete Impl;
+	Impl = nullptr;
+#endif // PLATFORM_LUMIN
+}
+
+bool FExternalTextureRendererGL::CreateImageKHR(void* NativeBufferHandle)
+{
+#if PLATFORM_LUMIN
+	Impl->Image = eglCreateImageKHR(LuminEGL::GetInstance()->GetDisplay(), EGL_NO_CONTEXT, EGL_NATIVE_BUFFER_ANDROID, (EGLClientBuffer)NativeBufferHandle, NULL);
+	if (Impl->Image == EGL_NO_IMAGE_KHR)
+	{
+		EGLint errorcode = eglGetError();
+		UE_LOG(LogMagicLeapHelperOpenGL, Error, TEXT("Failed to create EGLImage from the buffer. %d"), errorcode);
+		return false;
+	}
+
+	return true;
+#else
+	return false;
+#endif // PLATFORM_LUMIN
+}
+
+void FExternalTextureRendererGL::DestroyImageKHR()
+{
+#if PLATFORM_LUMIN
+	if (Impl->Image != EGL_NO_IMAGE_KHR)
+	{
+		eglDestroyImageKHR(eglGetCurrentDisplay(), Impl->Image);
+		Impl->Image = EGL_NO_IMAGE_KHR;
+	}
+#endif // PLATFORM_LUMIN
+}
+
+void FExternalTextureRendererGL::BindImageKHRToTexture(int32 TextureID)
+{
+#if PLATFORM_LUMIN
+	if (Impl->Image != EGL_NO_IMAGE_KHR)
+	{
+		glActiveTexture(GL_TEXTURE0);
+		glBindTexture(GL_TEXTURE_EXTERNAL_OES, TextureID);
+		glEGLImageTargetTexture2DOES(GL_TEXTURE_EXTERNAL_OES, Impl->Image);
+		glTexParameteri(GL_TEXTURE_EXTERNAL_OES, GL_TEXTURE_MIN_FILTER, GL_LINEAR);
+		glTexParameteri(GL_TEXTURE_EXTERNAL_OES, GL_TEXTURE_MAG_FILTER, GL_LINEAR);
+		glTexParameteri(GL_TEXTURE_EXTERNAL_OES, GL_TEXTURE_WRAP_S, GL_CLAMP_TO_EDGE);
+		glTexParameteri(GL_TEXTURE_EXTERNAL_OES, GL_TEXTURE_WRAP_T, GL_CLAMP_TO_EDGE);
+		glBindTexture(GL_TEXTURE_EXTERNAL_OES, 0);
+	}
+#endif // PLATFORM_LUMIN
+}
diff --git a/Engine/Plugins/Lumin/MagicLeap/Source/MagicLeapHelperOpenGL/Public/ExternalTextureRendererGL.h b/Engine/Plugins/Lumin/MagicLeap/Source/MagicLeapHelperOpenGL/Public/ExternalTextureRendererGL.h
new file mode 100644
--- /dev/null
+++ b/Engine/Plugins/Lumin/MagicLeap/Source/MagicLeapHelperOpenGL/Public/ExternalTextureRendererGL.h
@@ -0,0 +1,21 @@
+// Copyright 1998-2019 Epic Games, Inc. All Rights Reserved.
+
+#pragma once
+
+#include "CoreMinimal.h"
+
+class MAGICLEAPHELPEROPENGL_API FExternalTextureRendererGL
+{
+public:
+	FExternalTextureRendererGL();
+	virtual ~FExternalTextureRendererGL();
+
+	bool CreateImageKHR(void* NativeBufferHandle);
+	void DestroyImageKHR();
+	void BindImageKHRToTexture(int32 TextureID);
+
+#if PLATFORM_LUMIN
+private:
+	class FMagicLeapExternalTextureRendererGLImpl* Impl;
+#endif //PLATFORM_LUMIN
+};
diff --git a/Engine/Plugins/Lumin/MagicLeap/Source/MagicLeapHelperVulkan/MagicLeapHelperVulkan.Build.cs b/Engine/Plugins/Lumin/MagicLeap/Source/MagicLeapHelperVulkan/MagicLeapHelperVulkan.Build.cs
--- a/Engine/Plugins/Lumin/MagicLeap/Source/MagicLeapHelperVulkan/MagicLeapHelperVulkan.Build.cs
+++ b/Engine/Plugins/Lumin/MagicLeap/Source/MagicLeapHelperVulkan/MagicLeapHelperVulkan.Build.cs
@@ -15,16 +15,18 @@ namespace UnrealBuildTool.Rules
 				new string[] {
 					"Core",
 					"Engine",
+					"MLSDK",
 					"RHI",
 					"RenderCore",
 					"HeadMountedDisplay"
-                });
+				});
 
 			// TODO: Explore linking Unreal modules against a commong header and
 			// having a runtime dll linking against the library according to the platform.
 			if (Target.Platform != UnrealTargetPlatform.Mac)
 			{
 				PrivateDependencyModuleNames.Add("VulkanRHI");
+				PrivateDefinitions.Add("MLSDK_API_USE_VULKAN=1");
 				string EngineSourceDirectory = "../../../../Source";
 
 				if (Target.Platform == UnrealTargetPlatform.Linux)
@@ -43,7 +45,7 @@ namespace UnrealBuildTool.Rules
 						Path.Combine(EngineSourceDirectory, "Runtime/VulkanRHI/Private", ((Target.Platform == UnrealTargetPlatform.Win32 || Target.Platform == UnrealTargetPlatform.Win64) ? "Windows" : Target.Platform.ToString()))
 					});
 
-                AddEngineThirdPartyPrivateStaticDependencies(Target, "Vulkan");
+				AddEngineThirdPartyPrivateStaticDependencies(Target, "Vulkan");
 			}
 
 		}
diff --git a/Engine/Plugins/Lumin/MagicLeap/Source/MagicLeapHelperVulkan/Private/MagicLeapHelperVulkan.cpp b/Engine/Plugins/Lumin/MagicLeap/Source/MagicLeapHelperVulkan/Private/MagicLeapHelperVulkan.cpp
--- a/Engine/Plugins/Lumin/MagicLeap/Source/MagicLeapHelperVulkan/Private/MagicLeapHelperVulkan.cpp
+++ b/Engine/Plugins/Lumin/MagicLeap/Source/MagicLeapHelperVulkan/Private/MagicLeapHelperVulkan.cpp
@@ -13,9 +13,8 @@
 #include "VulkanUtil.h"
 #endif
 
-#if PLATFORM_LUMIN
-#include <ml_graphics_utils.h>
-#endif //PLATFORM_LUMIN
+#include "Lumin/CAPIShims/LuminAPIGraphics.h"
+#include "Lumin/CAPIShims/LuminAPIGraphicsUtils.h"
 
 DEFINE_LOG_CATEGORY_STATIC(LogMagicLeapHelperVulkan, Display, All);
 
@@ -26,7 +25,7 @@ IMPLEMENT_MODULE(FMagicLeapHelperVulkanPlugin, MagicLeapHelperVulkan);
 
 //////////////////////////////////////////////////////////////////////////
 
-void FMagicLeapHelperVulkan::BlitImage(uint64 SrcName, int32 SrcLevel, int32 SrcX, int32 SrcY, int32 SrcZ, int32 SrcWidth, int32 SrcHeight, int32 SrcDepth, uint64 DstName, int32 DstLevel, int32 DstX, int32 DstY, int32 DstZ, int32 DstWidth, int32 DstHeight, int32 DstDepth)
+void FMagicLeapHelperVulkan::BlitImage(uint64 SrcName, int32 SrcX, int32 SrcY, int32 SrcZ, int32 SrcWidth, int32 SrcHeight, int32 SrcDepth, uint64 DstName, int32 DstLayer, int32 DstX, int32 DstY, int32 DstZ, int32 DstWidth, int32 DstHeight, int32 DstDepth)
 {
 #if !PLATFORM_MAC
 	VkImage Src = (VkImage)SrcName;
@@ -63,20 +62,20 @@ void FMagicLeapHelperVulkan::BlitImage(uint64 SrcName, int32 SrcLevel, int32 Src
 	Region.srcSubresource.aspectMask = VK_IMAGE_ASPECT_COLOR_BIT;
 	Region.srcSubresource.layerCount = 1;
 	Region.dstOffsets[0].x = DstX;
+	// Unreal's viewport is bottom-left, ml_graphics is top-left so we invert the texture here.
 	Region.dstOffsets[0].y = DstY + DstHeight;
 	Region.dstOffsets[0].z = DstZ;
 	Region.dstOffsets[1].x = DstX + DstWidth;
 	Region.dstOffsets[1].y = DstY;
 	Region.dstOffsets[1].z = DstZ + DstDepth;
 	Region.dstSubresource.aspectMask = VK_IMAGE_ASPECT_COLOR_BIT;
-	Region.dstSubresource.baseArrayLayer = DstLevel;
+	Region.dstSubresource.baseArrayLayer = DstLayer;
 	Region.dstSubresource.layerCount = 1;
-	VulkanRHI::vkCmdBlitImage(CmdBuffer->GetHandle(), Src, VK_IMAGE_LAYOUT_TRANSFER_SRC_OPTIMAL, Dst, VK_IMAGE_LAYOUT_PRESENT_SRC_KHR, 1, &Region, VK_FILTER_LINEAR);
+	VulkanRHI::vkCmdBlitImage(CmdBuffer->GetHandle(), Src, VK_IMAGE_LAYOUT_TRANSFER_SRC_OPTIMAL, Dst, VK_IMAGE_LAYOUT_TRANSFER_DST_OPTIMAL, 1, &Region, VK_FILTER_LINEAR);
 #endif
 }
 
-
-void FMagicLeapHelperVulkan::TestClear(uint64 DstName)
+void FMagicLeapHelperVulkan::ClearImage(uint64 DstName, const FLinearColor& ClearColor, uint32 BaseMipLevel, uint32 LevelCount, uint32 BaseArrayLayer, uint32 LayerCount)
 {
 #if !PLATFORM_MAC
 	VkImage Dst = (VkImage)DstName;
@@ -87,16 +86,16 @@ void FMagicLeapHelperVulkan::TestClear(uint64 DstName)
 	FVulkanCmdBuffer* CmdBuffer = CmdBufferMgr->GetUploadCmdBuffer();
 
 	VkClearColorValue Color;
-	Color.float32[0] = 0;
-	Color.float32[1] = 0;
-	Color.float32[2] = 1;
-	Color.float32[3] = 1;
+	Color.float32[0] = ClearColor.R;
+	Color.float32[1] = ClearColor.G;
+	Color.float32[2] = ClearColor.B;
+	Color.float32[3] = ClearColor.A;
 	VkImageSubresourceRange Range;
 	Range.aspectMask = VK_IMAGE_ASPECT_COLOR_BIT;
-	Range.baseMipLevel = 0;
-	Range.levelCount = 1;
-	Range.baseArrayLayer = 0;
-	Range.layerCount = 2;
+	Range.baseMipLevel = BaseMipLevel;
+	Range.levelCount = LevelCount;
+	Range.baseArrayLayer = BaseArrayLayer;
+	Range.layerCount = LayerCount;
 	VulkanRHI::vkCmdClearColorImage(CmdBuffer->GetHandle(), Dst, VK_IMAGE_LAYOUT_PRESENT_SRC_KHR, &Color, 1, &Range);
 #endif
 }
@@ -160,6 +159,30 @@ uint64 FMagicLeapHelperVulkan::AliasImageSRGB(const uint64 Allocation, const uin
 #endif
 }
 
+void FMagicLeapHelperVulkan::DestroyImageSRGB(void* Image)
+{
+#if !PLATFORM_MAC
+	if (Image != VK_NULL_HANDLE)
+	{
+		FVulkanDynamicRHI* RHI = (FVulkanDynamicRHI*)GDynamicRHI;
+		FVulkanDevice* Device = RHI->GetDevice();
+		VulkanRHI::vkDestroyImage(Device->GetInstanceHandle(), reinterpret_cast<VkImage>(Image), nullptr);
+	}
+#endif
+}
+
+bool FMagicLeapHelperVulkan::GetVulkanInstanceExtensionsRequired(TArray<const ANSICHAR*>& Out)
+{
+#if PLATFORM_LUMIN
+	// Used inside ml_graphics. We get an error from validation layers for this extension not being enabled.
+	// TODO: Talk to graphics team for adding a MLGraphicsEnumerateRequiredVkInstanceExtensions() function so we dont have to hardcode this here.
+	// TODO: Talk to Epic if VULKAN_ENABLE_DESKTOP_HMD_SUPPORT can/should be enabled on Lumin since the extension is supported.
+	Out.Add(VK_KHR_GET_PHYSICAL_DEVICE_PROPERTIES_2_EXTENSION_NAME);
+#endif // PLATFORM_LUMIN
+
+	return true;
+}
+
 bool FMagicLeapHelperVulkan::GetVulkanDeviceExtensionsRequired(VkPhysicalDevice_T* pPhysicalDevice, TArray<const ANSICHAR*>& Out)
 {
 #if PLATFORM_LUMIN
@@ -199,6 +222,13 @@ bool FMagicLeapHelperVulkan::GetVulkanDeviceExtensionsRequired(VkPhysicalDevice_
 
 	const bool bFoundRequiredExtensions = (ExtensionsFound == RequiredExtensions.Num());
 	GSupportsImageExternal = bFoundRequiredExtensions; // This should probably be set by the vk rhi if the needed extensions are supported VK_KHR_external_memory?
+
+	// Used inside ml_graphics. We get an error from validation layers for these extensions not being enabled.
+	// Added after the checks for GSupportsImageExternal so we don't taint its flag with these unrelated extensions.
+	// TODO: Talk to graphics team for adding a MLGraphicsEnumerateRequiredVkDeviceExtensions() function so we dont have to hardcode this here.
+	Out.Add("VK_KHR_external_semaphore");
+	Out.Add("VK_KHR_external_semaphore_fd");
+
 	return bFoundRequiredExtensions;
 #endif //PLATFORM_LUMIN
 
