Copying //UE4/Partner-MagicLeap-Staging to Dev-VR (//UE4/Dev-VR)

diff --git a/Engine/Plugins/Lumin/MagicLeap/Source/MagicLeapHelperVulkan/Public/MagicLeapHelperVulkan.h b/Engine/Plugins/Lumin/MagicLeap/Source/MagicLeapHelperVulkan/Public/MagicLeapHelperVulkan.h
--- a/Engine/Plugins/Lumin/MagicLeap/Source/MagicLeapHelperVulkan/Public/MagicLeapHelperVulkan.h
+++ b/Engine/Plugins/Lumin/MagicLeap/Source/MagicLeapHelperVulkan/Public/MagicLeapHelperVulkan.h
@@ -8,10 +8,12 @@
 class MAGICLEAPHELPERVULKAN_API FMagicLeapHelperVulkan
 {
 public:
-	static void BlitImage(uint64 SrcName, int32 SrcLevel, int32 SrcX, int32 SrcY, int32 SrcZ, int32 SrcWidth, int32 SrcHeight, int32 SrcDepth, uint64 DstName, int32 DstLevel, int32 DstX, int32 DstY, int32 DstZ, int32 DstWidth, int32 DstHeight, int32 DstDepth);
+	static void BlitImage(uint64 SrcName, int32 SrcX, int32 SrcY, int32 SrcZ, int32 SrcWidth, int32 SrcHeight, int32 SrcDepth, uint64 DstName, int32 DstLayer, int32 DstX, int32 DstY, int32 DstZ, int32 DstWidth, int32 DstHeight, int32 DstDepth);
 	static void SignalObjects(uint64 SignalObject0, uint64 SignalObject1);
-	static void TestClear(uint64 Dest);
+	static void ClearImage(uint64 Dest, const FLinearColor& ClearColor, uint32 BaseMipLevel, uint32 LevelCount, uint32 BaseArrayLayer, uint32 LayerCount);
 	static uint64 AliasImageSRGB(const uint64 Allocation, const uint64 AllocationOffset, const uint32 Width, const uint32 Height);
+	static void DestroyImageSRGB(void* Dest);
+	static bool GetVulkanInstanceExtensionsRequired(TArray<const ANSICHAR*>& Out);
 	static bool GetVulkanDeviceExtensionsRequired(struct VkPhysicalDevice_T* pPhysicalDevice, TArray<const ANSICHAR*>& Out);
 	static bool GetMediaTexture(FTextureRHIRef& Result, FSamplerStateRHIRef& SamplerResult, const uint64 MediaTextureHandle);
 	static void AliasMediaTexture(FRHITexture* DestTexture, FRHITexture* SrcTexture);
diff --git a/Engine/Plugins/Lumin/MagicLeap/Source/MagicLeapIdentity/Private/MagicLeapIdentity.cpp b/Engine/Plugins/Lumin/MagicLeap/Source/MagicLeapIdentity/Private/MagicLeapIdentity.cpp
--- a/Engine/Plugins/Lumin/MagicLeap/Source/MagicLeapIdentity/Private/MagicLeapIdentity.cpp
+++ b/Engine/Plugins/Lumin/MagicLeap/Source/MagicLeapIdentity/Private/MagicLeapIdentity.cpp
@@ -5,41 +5,19 @@
 #include "Modules/ModuleManager.h"
 #include "HAL/PlatformProcess.h"
 #include "Misc/Paths.h"
-#include "MagicLeapPluginUtil.h"
-
-#if WITH_MLSDK
-#include "ml_identity.h"
-#endif //WITH_MLSDK
+#include "Lumin/CAPIShims/LuminAPIIdentity.h"
 
 DEFINE_LOG_CATEGORY_STATIC(LogMagicLeapIdentity, Display, All);
 
 class FMagicLeapIdentityPlugin : public IMagicLeapIdentityPlugin
 {
-public:
-	void StartupModule() override
-	{
-		IModuleInterface::StartupModule();
-		APISetup.Startup();
-#if WITH_MLSDK
-		APISetup.LoadDLL(TEXT("ml_identity"));
-#endif //WITH_MLSDK
-	}
-
-	void ShutdownModule() override
-	{
-		APISetup.Shutdown();
-		IModuleInterface::ShutdownModule();
-	}
-
-private:
-	FMagicLeapAPISetup APISetup;
 };
 
 IMPLEMENT_MODULE(FMagicLeapIdentityPlugin, MagicLeapIdentity);
 
 //////////////////////////////////////////////////////////////////////////
 
-class FIdentityImpl
+class FMagicLeapIdentityImpl
 {
 #if WITH_MLSDK
 public:
@@ -149,7 +127,7 @@ public:
 };
 
 UMagicLeapIdentity::UMagicLeapIdentity()
-	: Impl(new FIdentityImpl())
+	: Impl(new FMagicLeapIdentityImpl())
 {}
 
 UMagicLeapIdentity::~UMagicLeapIdentity()
@@ -192,7 +170,7 @@ void UMagicLeapIdentity::GetAllAvailableAttributesAsync(const FAvailableIdentity
 	}
 	else
 	{
-		UE_LOG(LogMagicLeapIdentity, Error, TEXT("MLIdentityGetAttributeNamesAsync failed with error %d!"), Result);
+		UE_LOG(LogMagicLeapIdentity, Error, TEXT("MLIdentityGetAttributeNamesAsync failed with error %s!"), UTF8_TO_TCHAR(MLIdentityGetResultString(Result)));
 	}
 #endif //WITH_MLSDK
 }
@@ -260,11 +238,11 @@ EMagicLeapIdentityError UMagicLeapIdentity::RequestAttributeValueAsync(const TAr
 		MLResult RequestAtrrValRes = MLIdentityRequestAttributeValuesAsync(profile, &InvokeFuture);
 		if (RequestAtrrValRes == MLResult_Ok)
 		{
-			Impl->AllRequestAttribsFutures.Add(InvokeFuture, FIdentityImpl::FRequestAttribData(profile, ResultDelegate));
+			Impl->AllRequestAttribsFutures.Add(InvokeFuture, FMagicLeapIdentityImpl::FRequestAttribData(profile, ResultDelegate));
 		}
 		else
 		{
-			UE_LOG(LogMagicLeapIdentity, Error, TEXT("MLIdentityRequestAttributeValuesAsync failed with error %d!"), RequestAtrrValRes);
+			UE_LOG(LogMagicLeapIdentity, Error, TEXT("MLIdentityRequestAttributeValuesAsync failed with error %s!"), UTF8_TO_TCHAR(MLIdentityGetResultString(RequestAtrrValRes)));
 		}
 	}
 
diff --git a/Engine/Plugins/Lumin/MagicLeap/Source/MagicLeapIdentity/Public/MagicLeapIdentity.h b/Engine/Plugins/Lumin/MagicLeap/Source/MagicLeapIdentity/Public/MagicLeapIdentity.h
--- a/Engine/Plugins/Lumin/MagicLeap/Source/MagicLeapIdentity/Public/MagicLeapIdentity.h
+++ b/Engine/Plugins/Lumin/MagicLeap/Source/MagicLeapIdentity/Public/MagicLeapIdentity.h
@@ -89,5 +89,5 @@ public:
 	virtual UWorld* GetTickableGameObjectWorld() const override;
 
 private:
-	class FIdentityImpl *Impl;
+	class FMagicLeapIdentityImpl *Impl;
 };
diff --git a/Engine/Plugins/Lumin/MagicLeap/Source/MagicLeapImageTracker/MagicLeapImageTracker.Build.cs b/Engine/Plugins/Lumin/MagicLeap/Source/MagicLeapImageTracker/MagicLeapImageTracker.Build.cs
new file mode 100644
--- /dev/null
+++ b/Engine/Plugins/Lumin/MagicLeap/Source/MagicLeapImageTracker/MagicLeapImageTracker.Build.cs
@@ -0,0 +1,32 @@
+// Copyright 1998-2019 Epic Games, Inc. All Rights Reserved.
+
+using System.IO;
+
+namespace UnrealBuildTool.Rules
+{
+	public class MagicLeapImageTracker : ModuleRules
+	{
+		public MagicLeapImageTracker(ReadOnlyTargetRules Target) : base(Target)
+		{
+			PublicIncludePaths.Add(Path.Combine(ModuleDirectory, "Public"));
+
+			PublicDependencyModuleNames.AddRange(
+				new string[]
+				{
+					"Core",
+					"CoreUObject",
+					"Engine",
+				}
+			);
+			
+			PrivateDependencyModuleNames.AddRange(
+				new string[]
+				{
+					"LuminRuntimeSettings",
+					"MLSDK",
+					"MagicLeap"
+				}
+			);
+		}
+	}
+}
diff --git a/Engine/Plugins/Lumin/MagicLeap/Source/MagicLeapImageTracker/Private/MagicLeapImageTrackerComponent.cpp b/Engine/Plugins/Lumin/MagicLeap/Source/MagicLeapImageTracker/Private/MagicLeapImageTrackerComponent.cpp
new file mode 100644
--- /dev/null
+++ b/Engine/Plugins/Lumin/MagicLeap/Source/MagicLeapImageTracker/Private/MagicLeapImageTrackerComponent.cpp
@@ -0,0 +1,125 @@
+// Copyright 1998-2019 Epic Games, Inc. All Rights Reserved.
+
+#include "MagicLeapImageTrackerComponent.h"
+#include "IMagicLeapPlugin.h"
+#include "MagicLeapImageTrackerModule.h"
+#include "MagicLeapImageTrackerRunnable.h"
+#if WITH_EDITOR
+#include "Editor.h"
+#endif
+
+UMagicLeapImageTrackerComponent::UMagicLeapImageTrackerComponent()
+: bIsTracking(false)
+#if WITH_EDITOR
+, TextureBeforeEdit(nullptr)
+#endif // WITH_EDITOR
+{
+	// Make sure this component ticks
+	PrimaryComponentTick.bCanEverTick = true;
+	PrimaryComponentTick.bStartWithTickEnabled = true;
+	PrimaryComponentTick.TickGroup = TG_PrePhysics;
+
+	bAutoActivate = true;
+}
+
+void UMagicLeapImageTrackerComponent::EndPlay(const EEndPlayReason::Type EndPlayReason)
+{
+	RemoveTargetAsync();
+	Super::EndPlay(EndPlayReason);
+}
+
+void UMagicLeapImageTrackerComponent::TickComponent(float DeltaTime, enum ELevelTick TickType, FActorComponentTickFunction *ThisTickFunction)
+{
+	Super::TickComponent(DeltaTime, TickType, ThisTickFunction);
+
+	if (!bIsTracking)
+	{
+		UTexture2D* TempTargetImageTexture = TargetImageTexture;
+		TargetImageTexture = nullptr;
+		bIsTracking = SetTargetAsync(TempTargetImageTexture);
+	}
+	else
+	{
+		FVector Location;
+		FRotator Rotation;
+		if (GetMagicLeapImageTrackerModule().TryGetRelativeTransform(Name, Location, Rotation))
+		{
+			this->SetRelativeLocationAndRotation(Location, Rotation);
+		}
+	}
+}
+
+bool UMagicLeapImageTrackerComponent::SetTargetAsync(UTexture2D* ImageTarget)
+{
+	if (ImageTarget == nullptr)
+	{
+		UE_LOG(LogMagicLeapImageTracker, Warning, TEXT("ImageTarget is NULL!."));
+		return false;
+	}
+
+	if (ImageTarget && !(ImageTarget->GetPixelFormat() == EPixelFormat::PF_R8G8B8A8 || ImageTarget->GetPixelFormat() == EPixelFormat::PF_B8G8R8A8))
+	{
+		UE_LOG(LogMagicLeapImageTracker, Error, TEXT("Cannot set texture %s as it uses an invalid pixel format!  Valid formats are R8B8G8A8 or B8G8R8A8"), *ImageTarget->GetName());
+		return false;
+	}
+
+	if (ImageTarget == TargetImageTexture)
+	{
+		UE_LOG(LogMagicLeapImageTracker, Error, TEXT("Skipped setting %s as it is already being used as the current image target"), *ImageTarget->GetName());
+		return false;
+	}
+
+	TargetImageTexture = ImageTarget;
+	if (Name.Len() == 0) Name = GetFullGroupName(true);
+
+	FMagicLeapImageTrackerTarget Target;
+	Target.Name = Name;
+#if WITH_MLSDK
+	Target.Settings.longer_dimension = LongerDimension / IMagicLeapPlugin::Get().GetWorldToMetersScale();
+	Target.Settings.is_stationary = bIsStationary;
+#endif // WITH_MLSDK
+	Target.Texture = TargetImageTexture;
+	Target.bUseUnreliablePose = bUseUnreliablePose;
+	Target.OnSetImageTargetSucceeded = OnSetImageTargetSucceeded;
+	Target.OnSetImageTargetFailed = OnSetImageTargetFailed;
+	Target.OnImageTargetFound = OnImageTargetFound;
+	Target.OnImageTargetLost = OnImageTargetLost;
+	Target.OnImageTargetUnreliableTracking = OnImageTargetUnreliableTracking;
+
+	GetMagicLeapImageTrackerModule().SetTargetAsync(Target);
+	return true;
+}
+
+bool UMagicLeapImageTrackerComponent::RemoveTargetAsync()
+{
+	bIsTracking = false;
+	return GetMagicLeapImageTrackerModule().RemoveTargetAsync(Name);
+}
+
+#if WITH_EDITOR
+void UMagicLeapImageTrackerComponent::PreEditChange(UProperty* PropertyAboutToChange)
+{
+	if ((PropertyAboutToChange != nullptr) && (PropertyAboutToChange->GetFName() == GET_MEMBER_NAME_CHECKED(UMagicLeapImageTrackerComponent, TargetImageTexture)))
+	{
+		TextureBeforeEdit = TargetImageTexture;
+	}
+
+	Super::PreEditChange(PropertyAboutToChange);
+}
+
+void UMagicLeapImageTrackerComponent::PostEditChangeProperty(FPropertyChangedEvent& PropertyChangedEvent)
+{
+	const FName PropertyName = (PropertyChangedEvent.Property != nullptr) ? PropertyChangedEvent.Property->GetFName() : NAME_None;
+
+	if (PropertyName == GET_MEMBER_NAME_CHECKED(UMagicLeapImageTrackerComponent, TargetImageTexture))
+	{
+		if (TargetImageTexture && !(TargetImageTexture->GetPixelFormat() == EPixelFormat::PF_R8G8B8A8 || TargetImageTexture->GetPixelFormat() == EPixelFormat::PF_B8G8R8A8))
+		{
+			UE_LOG(LogMagicLeapImageTracker, Error, TEXT("[UImageTrackerComponent] Cannot set texture %s as it uses an invalid pixel format!  Valid formats are R8B8G8A8 or B8G8R8A8"), *TargetImageTexture->GetName());
+			TargetImageTexture = TextureBeforeEdit;
+		}
+	}
+
+	Super::PostEditChangeProperty(PropertyChangedEvent);
+}
+#endif
diff --git a/Engine/Plugins/Lumin/MagicLeap/Source/MagicLeapImageTracker/Private/MagicLeapImageTrackerFunctionLibrary.cpp b/Engine/Plugins/Lumin/MagicLeap/Source/MagicLeapImageTracker/Private/MagicLeapImageTrackerFunctionLibrary.cpp
new file mode 100644
--- /dev/null
+++ b/Engine/Plugins/Lumin/MagicLeap/Source/MagicLeapImageTracker/Private/MagicLeapImageTrackerFunctionLibrary.cpp
@@ -0,0 +1,24 @@
+// Copyright 1998-2019 Epic Games, Inc. All Rights Reserved.
+
+#include "MagicLeapImageTrackerFunctionLibrary.h"
+#include "MagicLeapImageTrackerModule.h"
+
+void UMagicLeapImageTrackerFunctionLibrary::SetMaxSimultaneousTargets(int32 MaxSimultaneousTargets)
+{
+	GetMagicLeapImageTrackerModule().SetMaxSimultaneousTargets(MaxSimultaneousTargets);
+}
+
+int32 UMagicLeapImageTrackerFunctionLibrary::GetMaxSimultaneousTargets()
+{
+	return GetMagicLeapImageTrackerModule().GetMaxSimultaneousTargets();
+}
+
+void UMagicLeapImageTrackerFunctionLibrary::EnableImageTracking(bool bEnable)
+{
+	GetMagicLeapImageTrackerModule().SetImageTrackerEnabled(bEnable);
+}
+
+bool UMagicLeapImageTrackerFunctionLibrary::IsImageTrackingEnabled()
+{
+	return GetMagicLeapImageTrackerModule().GetImageTrackerEnabled();
+}
diff --git a/Engine/Plugins/Lumin/MagicLeap/Source/MagicLeapImageTracker/Private/MagicLeapImageTrackerModule.cpp b/Engine/Plugins/Lumin/MagicLeap/Source/MagicLeapImageTracker/Private/MagicLeapImageTrackerModule.cpp
new file mode 100644
--- /dev/null
+++ b/Engine/Plugins/Lumin/MagicLeap/Source/MagicLeapImageTracker/Private/MagicLeapImageTrackerModule.cpp
@@ -0,0 +1,101 @@
+// Copyright 1998-2019 Epic Games, Inc. All Rights Reserved.
+
+#include "MagicLeapImageTrackerModule.h"
+#include "MagicLeapImageTrackerRunnable.h"
+
+FMagicLeapImageTrackerModule::FMagicLeapImageTrackerModule()
+: Runnable(nullptr)
+{
+}
+
+void FMagicLeapImageTrackerModule::StartupModule()
+{
+	IMagicLeapImageTrackerModule::StartupModule();
+	TickDelegate = FTickerDelegate::CreateRaw(this, &FMagicLeapImageTrackerModule::Tick);
+	Runnable = new FMagicLeapImageTrackerRunnable;
+}
+
+void FMagicLeapImageTrackerModule::ShutdownModule()
+{
+	if (TickDelegateHandle.IsValid())
+	{
+		FTicker::GetCoreTicker().RemoveTicker(TickDelegateHandle);
+	}
+	FMagicLeapImageTrackerRunnable* InRunnable = Runnable;
+	Runnable = nullptr;
+	AsyncTask(ENamedThreads::AnyBackgroundThreadNormalTask, [InRunnable]()
+	{
+		delete InRunnable;
+	});
+
+	IModuleInterface::ShutdownModule();
+}
+
+bool FMagicLeapImageTrackerModule::Tick(float DeltaTime)
+{
+	FMagicLeapImageTrackerTask CompletedTask;
+	if (Runnable->TryGetCompletedTask(CompletedTask))
+	{
+		FMagicLeapImageTrackerTarget& ImageTrackerTarget = CompletedTask.Target;
+		switch (CompletedTask.Type)
+		{
+		case FMagicLeapImageTrackerTask::EType::TargetCreateFailed:
+		{
+			ImageTrackerTarget.OnSetImageTargetFailed.Broadcast();
+		}
+		break;
+
+		case FMagicLeapImageTrackerTask::EType::TargetCreateSucceeded:
+		{
+			ImageTrackerTarget.OnSetImageTargetSucceeded.Broadcast();
+		}
+		break;
+		}
+	}
+
+	Runnable->UpdateTargetsMainThread();
+
+	return true;
+}
+
+void FMagicLeapImageTrackerModule::SetTargetAsync(const FMagicLeapImageTrackerTarget& ImageTarget)
+{
+	Runnable->SetTargetAsync(ImageTarget);
+	if (!TickDelegateHandle.IsValid())
+	{
+		TickDelegateHandle = FTicker::GetCoreTicker().AddTicker(TickDelegate);
+	}
+}
+
+bool FMagicLeapImageTrackerModule::RemoveTargetAsync(const FString& TargetName)
+{
+	return Runnable->RemoveTargetAsync(TargetName);
+}
+
+uint32 FMagicLeapImageTrackerModule::GetMaxSimultaneousTargets() const
+{
+	return Runnable->GetMaxSimultaneousTargets();
+}
+
+void FMagicLeapImageTrackerModule::SetMaxSimultaneousTargets(uint32 NewNumTargets)
+{
+	Runnable->SetMaxSimultaneousTargets(NewNumTargets);
+}
+
+bool FMagicLeapImageTrackerModule::GetImageTrackerEnabled() const
+{
+	return Runnable->GetImageTrackerEnabled();
+}
+
+void FMagicLeapImageTrackerModule::SetImageTrackerEnabled(bool bEnabled)
+{
+	Runnable->SetImageTrackerEnabled(bEnabled);
+}
+
+bool FMagicLeapImageTrackerModule::TryGetRelativeTransform(const FString& TargetName, FVector& OutLocation, FRotator& OutRotation)
+{
+	return Runnable->TryGetRelativeTransformMainThread(TargetName, OutLocation, OutRotation);
+}
+
+
+IMPLEMENT_MODULE(FMagicLeapImageTrackerModule, MagicLeapImageTracker);
diff --git a/Engine/Plugins/Lumin/MagicLeap/Source/MagicLeapImageTracker/Private/MagicLeapImageTrackerModule.h b/Engine/Plugins/Lumin/MagicLeap/Source/MagicLeapImageTracker/Private/MagicLeapImageTrackerModule.h
new file mode 100644
--- /dev/null
+++ b/Engine/Plugins/Lumin/MagicLeap/Source/MagicLeapImageTracker/Private/MagicLeapImageTrackerModule.h
@@ -0,0 +1,33 @@
+// Copyright 1998-2019 Epic Games, Inc. All Rights Reserved.
+
+#pragma once
+
+#include "Containers/Ticker.h"
+#include "IMagicLeapImageTrackerModule.h"
+#include "MagicLeapImageTrackerTypes.h"
+
+class FMagicLeapImageTrackerModule : public IMagicLeapImageTrackerModule
+{
+public:
+	FMagicLeapImageTrackerModule();
+	void StartupModule() override;
+	void ShutdownModule() override;
+	bool Tick(float DeltaTime);
+	void SetTargetAsync(const FMagicLeapImageTrackerTarget& ImageTarget);
+	bool RemoveTargetAsync(const FString& InName);
+	uint32 GetMaxSimultaneousTargets() const;
+	void SetMaxSimultaneousTargets(uint32 NewNumTargets);
+	bool GetImageTrackerEnabled() const;
+	void SetImageTrackerEnabled(bool bEnabled);
+	bool TryGetRelativeTransform(const FString& TargetName, FVector& OutLocation, FRotator& OutRotation);
+
+private:
+	FTickerDelegate TickDelegate;
+	FDelegateHandle TickDelegateHandle;
+	class FMagicLeapImageTrackerRunnable* Runnable;
+};
+
+inline FMagicLeapImageTrackerModule& GetMagicLeapImageTrackerModule()
+{
+	return FModuleManager::Get().GetModuleChecked<FMagicLeapImageTrackerModule>("MagicLeapImageTracker");
+}
diff --git a/Engine/Plugins/Lumin/MagicLeap/Source/MagicLeapImageTracker/Private/MagicLeapImageTrackerRunnable.cpp b/Engine/Plugins/Lumin/MagicLeap/Source/MagicLeapImageTracker/Private/MagicLeapImageTrackerRunnable.cpp
new file mode 100644
--- /dev/null
+++ b/Engine/Plugins/Lumin/MagicLeap/Source/MagicLeapImageTracker/Private/MagicLeapImageTrackerRunnable.cpp
@@ -0,0 +1,483 @@
+// Copyright 1998-2019 Epic Games, Inc. All Rights Reserved.
+
+#include "MagicLeapImageTrackerRunnable.h"
+
+DEFINE_LOG_CATEGORY(LogMagicLeapImageTracker);
+
+FMagicLeapImageTrackerRunnable::FMagicLeapImageTrackerRunnable()
+: FMagicLeapRunnable({ EMagicLeapPrivilege::CameraCapture }, TEXT("FMLImageTrackerRunnable"), EThreadPriority::TPri_AboveNormal)
+#if WITH_MLSDK
+, ImageTracker(ML_INVALID_HANDLE)
+#endif // WITH_MLSDK
+{
+#if WITH_MLSDK
+	FMemory::Memset(&Settings, 0, sizeof(Settings));
+#endif // WITH_MLSDK
+};
+
+void FMagicLeapImageTrackerRunnable::Exit()
+{
+#if WITH_MLSDK
+	{
+		FScopeLock Lock(&TargetsMutex);
+		for (auto& Target : TrackedImageTargets)
+		{
+			RemoveTarget(Target.Value);
+		}
+		TrackedImageTargets.Empty();
+	}
+
+	DestroyTracker();
+#endif // WITH_MLSDK
+}
+
+void FMagicLeapImageTrackerRunnable::Pause()
+{
+#if WITH_MLSDK
+	AppEventHandler.SetWasSystemEnabledOnPause(GetImageTrackerEnabled());
+
+	if (!AppEventHandler.WasSystemEnabledOnPause())
+	{
+		UE_LOG(LogMagicLeapImageTracker, Log, TEXT("Image tracking was not enabled at time of application pause."));
+	}
+	else
+	{
+		if (!MLHandleIsValid(ImageTracker))
+		{
+			UE_LOG(LogMagicLeapImageTracker, Error, TEXT("Image tracker was invalid on application pause."));
+		}
+		else
+		{
+			SetImageTrackerEnabled(false);
+		}
+	}
+#endif // WITH_MLSDK
+}
+
+void FMagicLeapImageTrackerRunnable::Resume()
+{
+#if WITH_MLSDK
+	if (!AppEventHandler.WasSystemEnabledOnPause())
+	{
+		UE_LOG(LogMagicLeapImageTracker, Log, TEXT("Not resuming image tracker as it was not enabled at time of application pause."));
+	}
+	else
+	{
+		if (!MLHandleIsValid(ImageTracker))
+		{
+			UE_LOG(LogMagicLeapImageTracker, Error, TEXT("Image tracker was invalid on application resume."));
+		}
+		else
+		{
+			SetImageTrackerEnabled(true);
+		}
+	}
+#endif // WITH_MLSDK
+}
+
+void FMagicLeapImageTrackerRunnable::SetTargetAsync(const FMagicLeapImageTrackerTarget& ImageTarget)
+{
+	FMagicLeapImageTrackerTask SetTargetMsg;
+	SetTargetMsg.Type = FMagicLeapImageTrackerTask::EType::TryCreateTarget;
+	SetTargetMsg.Target = ImageTarget;
+
+	PushNewTask(SetTargetMsg);
+}
+
+bool FMagicLeapImageTrackerRunnable::RemoveTargetAsync(const FString& InName)
+{
+	FMagicLeapImageTrackerTarget* Target = TrackedImageTargets.Find(InName);
+	if (Target)
+	{
+		FMagicLeapImageTrackerTask RemoveTargetMsg;
+		RemoveTargetMsg.Type = FMagicLeapImageTrackerTask::EType::TryRemoveTarget;
+		RemoveTargetMsg.Target = *Target;
+
+		PushNewTask(RemoveTargetMsg);
+		return true;
+	}
+
+	UE_LOG(LogMagicLeapImageTracker, Warning, TEXT("RemoveTargetAsync failed to resolve target '%s'!"), *InName);
+	return false;
+}
+
+bool FMagicLeapImageTrackerRunnable::ProcessCurrentTask()
+{
+	bool bProcessedTask = false;
+
+	switch (CurrentTask.Type)
+	{
+	case FMagicLeapImageTrackerTask::EType::None: break;
+
+	case FMagicLeapImageTrackerTask::EType::TryCreateTarget:
+	{
+#if WITH_MLSDK
+		SetTarget();
+#endif // WITH_MLSDK
+		bProcessedTask = true;
+	}
+	break;
+
+	case FMagicLeapImageTrackerTask::EType::TryRemoveTarget:
+	{
+#if WITH_MLSDK
+		FScopeLock Lock(&TargetsMutex);
+		RemoveTarget(CurrentTask.Target, true);
+#endif // WITH_MLSDK
+		bProcessedTask = true;
+	}
+	break;
+	}
+
+	return bProcessedTask;
+}
+
+uint32 FMagicLeapImageTrackerRunnable::GetMaxSimultaneousTargets() const
+{
+#if WITH_MLSDK
+	return static_cast<uint32>(FPlatformAtomics::AtomicRead(reinterpret_cast<const volatile int32*>(&Settings.max_simultaneous_targets)));
+#else
+	return 0;
+#endif // WITH_MLSDK
+}
+
+void FMagicLeapImageTrackerRunnable::SetMaxSimultaneousTargets(uint32 NewNumTargets)
+{
+#if WITH_MLSDK
+	// run on first available background thread rather than our worker thread in order
+	// to avoid a potential task queue.
+	AsyncTask(ENamedThreads::AnyBackgroundThreadNormalTask, [this, NewNumTargets]()
+	{
+		if (!MLHandleIsValid(GetHandle()))
+		{
+			CreateTracker();
+		}
+		FScopeLock Lock(&TrackerMutex);
+		if (AppEventHandler.GetPrivilegeStatus(EMagicLeapPrivilege::CameraCapture) == MagicLeap::EPrivilegeState::Granted)
+		{
+			Settings.max_simultaneous_targets = NewNumTargets;
+			MLResult Result = MLImageTrackerUpdateSettings(ImageTracker, &Settings);
+			if (Result != MLResult_Ok)
+			{
+				UE_LOG(LogMagicLeapImageTracker, Error, TEXT("MLImageTrackerUpdateSettings failed with error '%s'!"), UTF8_TO_TCHAR(MLGetResultString(Result)));
+			}
+		}
+		else
+		{
+			UE_LOG(LogMagicLeapImageTracker, Log, TEXT("Image tracking settings failed to update due to lack of privilege!"));
+		}
+	});
+#else
+	(void)NewNumTargets;
+#endif // WITH_MLSDK
+}
+
+bool FMagicLeapImageTrackerRunnable::GetImageTrackerEnabled() const
+{
+#if WITH_MLSDK
+	return static_cast<bool>(FPlatformAtomics::AtomicRead(reinterpret_cast<const volatile int32*>(&Settings.enable_image_tracking))) && MLHandleIsValid(GetHandle());
+#else 
+	return false;
+#endif // WITH_MLSDK
+}
+
+void FMagicLeapImageTrackerRunnable::SetImageTrackerEnabled(bool bEnabled)
+{
+#if WITH_MLSDK
+	// run on first available background thread rather than our worker thread in order
+	// to avoid a potential task queue.
+	AsyncTask(ENamedThreads::AnyBackgroundThreadNormalTask, [this, bEnabled]()
+	{
+		if (MLHandleIsValid(GetHandle()))
+		{
+			FScopeLock Lock(&TrackerMutex);
+			if (AppEventHandler.GetPrivilegeStatus(EMagicLeapPrivilege::CameraCapture) == MagicLeap::EPrivilegeState::Granted)
+			{
+				Settings.enable_image_tracking = bEnabled;
+				MLResult Result = MLImageTrackerUpdateSettings(ImageTracker, &Settings);
+				if (Result != MLResult_Ok)
+				{
+					UE_LOG(LogMagicLeapImageTracker, Error, TEXT("MLImageTrackerUpdateSettings failed with error '%s'!"), UTF8_TO_TCHAR(MLGetResultString(Result)));
+				}
+			}
+			else
+			{
+				UE_LOG(LogMagicLeapImageTracker, Log, TEXT("Image tracking settings failed to update due to lack of privilege!"));
+			}
+		}
+		else if (bEnabled)
+		{
+			CreateTracker();
+		}
+	});
+#else
+	(void)bEnabled;
+#endif // WITH_MLSDK
+}
+
+void FMagicLeapImageTrackerRunnable::UpdateTargetsMainThread()
+{
+#if WITH_MLSDK
+	if (!(IMagicLeapPlugin::Get().IsMagicLeapHMDValid()))
+	{
+		return;
+	}
+
+	FScopeLock Lock(&TargetsMutex);
+	for (auto& TargetEntry : TrackedImageTargets)
+	{
+		FMagicLeapImageTrackerTarget& Target = TargetEntry.Value;
+		checkf(MLHandleIsValid(GetHandle()), TEXT("ImageTracker handle is invalid!"));
+		MLImageTrackerTargetResult TrackingStatus;
+		MLResult Result = MLImageTrackerGetTargetResult(GetHandle(), Target.Handle, &TrackingStatus);
+		if (Result != MLResult_Ok)
+		{
+			UE_LOG(LogMagicLeapImageTracker, Warning, TEXT("MLImageTrackerGetTargetResult failed with error '%s'."), UTF8_TO_TCHAR(MLGetResultString(Result)));
+			TrackingStatus.status = MLImageTrackerTargetStatus_NotTracked;
+		}
+
+		if (TrackingStatus.status == MLImageTrackerTargetStatus_NotTracked)
+		{
+			if (Target.OldTrackingStatus.status != MLImageTrackerTargetStatus_NotTracked)
+			{
+				Target.OnImageTargetLost.Broadcast();
+			}
+		}
+		else
+		{
+			EMagicLeapTransformFailReason FailReason = EMagicLeapTransformFailReason::None;
+			FTransform Pose = FTransform::Identity;
+			if (IMagicLeapPlugin::Get().GetTransform(Target.Data.coord_frame_target, Pose, FailReason))
+			{
+				Pose.ConcatenateRotation(FQuat(FVector(0, 0, 1), PI));
+				if (TrackingStatus.status == MLImageTrackerTargetStatus_Unreliable)
+				{
+					Target.UnreliableLocation = Pose.GetTranslation();
+					Target.UnreliableRotation = Pose.Rotator();
+					Target.OnImageTargetUnreliableTracking.Broadcast(Target.Location, Target.Rotation, Target.UnreliableLocation, Target.UnreliableRotation);
+				}
+				else
+				{
+					Target.Location = Pose.GetTranslation();
+					Target.Rotation = Pose.Rotator();
+					if (Target.OldTrackingStatus.status != MLImageTrackerTargetStatus_Tracked)
+					{
+						Target.OnImageTargetFound.Broadcast();
+					}
+				}
+			}
+			else
+			{
+				if (FailReason == EMagicLeapTransformFailReason::NaNsInTransform)
+				{
+					UE_LOG(LogMagicLeapImageTracker, Error, TEXT("NaNs in image tracker target transform."));
+				}
+				TrackingStatus.status = MLImageTrackerTargetStatus_NotTracked;
+				if (Target.OldTrackingStatus.status != MLImageTrackerTargetStatus_NotTracked)
+				{
+					Target.OnImageTargetLost.Broadcast();
+				}
+			}
+		}
+
+		Target.OldTrackingStatus = TrackingStatus;
+	}
+#endif // WITH_MLSDK
+}
+
+bool FMagicLeapImageTrackerRunnable::TryGetRelativeTransformMainThread(const FString& TargetName, FVector& OutLocation, FRotator& OutRotation)
+{
+#if WITH_MLSDK
+	if (!(IMagicLeapPlugin::Get().IsMagicLeapHMDValid()))
+	{
+		return false;
+	}
+
+	FScopeLock Lock(&TargetsMutex);
+	FMagicLeapImageTrackerTarget* Target = TrackedImageTargets.Find(TargetName);
+	if (Target)
+	{
+		if (Target->OldTrackingStatus.status == MLImageTrackerTargetStatus_Unreliable)
+		{
+			if (Target->bUseUnreliablePose)
+			{
+				OutLocation = Target->UnreliableLocation;
+				OutRotation = Target->UnreliableRotation;
+				return true;
+			}
+
+			return false;
+		}
+		else if (Target->OldTrackingStatus.status == MLImageTrackerTargetStatus_Tracked)
+		{
+			OutLocation = Target->Location;
+			OutRotation = Target->Rotation;
+			return true;
+		}
+	}
+#endif // WITH_MLSDK
+	return false;
+}
+
+
+#if WITH_MLSDK
+bool FMagicLeapImageTrackerRunnable::CreateTracker()
+{
+	if (MagicLeap::EPrivilegeState::Granted != AppEventHandler.GetPrivilegeStatus(EMagicLeapPrivilege::CameraCapture))
+	{
+		return false;
+	}
+
+	MLImageTrackerSettings TempSettings;
+	// don't allow atomic reads of Settings while we are creating the tracker
+	MLResult Result = MLImageTrackerInitSettings(&TempSettings);
+	if (Result != MLResult_Ok)
+	{
+		UE_LOG(LogMagicLeapImageTracker, Error, TEXT("MLImageTrackerInitSettings failed with error '%s'!"), UTF8_TO_TCHAR(MLGetResultString(Result)));
+		return false;
+	}
+
+	MLHandle TempHandle = ML_INVALID_HANDLE;
+	// don't allow atomic reads of the tracker handle while we are creating the tracker
+	Result = MLImageTrackerCreate(&TempSettings, &TempHandle);
+
+	if (Result != MLResult_Ok)
+	{
+		UE_LOG(LogMagicLeapImageTracker, Error, TEXT("MLImageTrackerCreate failed with error '%s'!"), UTF8_TO_TCHAR(MLGetResultString(Result)));
+		return false;
+	}
+
+	FScopeLock Lock(&TrackerMutex);
+	Settings = TempSettings;
+	ImageTracker = TempHandle;
+	return true;
+}
+
+void FMagicLeapImageTrackerRunnable::DestroyTracker()
+{
+	FScopeLock Lock(&TrackerMutex);
+	if (MLHandleIsValid(ImageTracker))
+	{
+		MLHandle TempHandle = ImageTracker;
+		// don't allow atomic reads of the tracker handle while we are destroying the tracker
+		ImageTracker = ML_INVALID_HANDLE;
+		MLResult Result = MLImageTrackerDestroy(TempHandle);
+
+		if (Result != MLResult_Ok)
+		{
+			UE_LOG(LogMagicLeapImageTracker, Error, TEXT("MLImageTrackerDestroy failed with error '%s'!"), UTF8_TO_TCHAR(MLGetResultString(Result)));
+		}
+
+		ImageTracker = ML_INVALID_HANDLE;
+	}
+}
+
+void FMagicLeapImageTrackerRunnable::SetTarget()
+{
+	if (!MLHandleIsValid(ImageTracker))
+	{
+		if (!CreateTracker())
+		{
+			return;
+		}
+	}
+
+	FMagicLeapImageTrackerTarget& Target = CurrentTask.Target;
+	{
+		FScopeLock Lock(&TargetsMutex);
+		FMagicLeapImageTrackerTarget* ExistingTarget = TrackedImageTargets.Find(Target.Name);
+		if (ExistingTarget != nullptr && MLHandleIsValid(ExistingTarget->Handle))
+		{
+			RemoveTarget(*ExistingTarget);
+		}
+	}
+
+	UE_LOG(LogMagicLeapImageTracker, Warning, TEXT("SetTarget for %s"), *Target.Name);
+	FMagicLeapImageTrackerTask TargetCreateTask;
+	Target.Settings.name = TCHAR_TO_UTF8(*Target.Name);
+
+	if (Target.Texture->PlatformData->Mips.Num() == 0)
+	{
+		UE_LOG(LogMagicLeapImageTracker, Error, TEXT("Texture for target %s has no mip maps."), *Target.Name);
+		TargetCreateTask.Type = FMagicLeapImageTrackerTask::EType::TargetCreateFailed;
+		PushCompletedTask(TargetCreateTask);
+		return;
+	}
+
+	FTexture2DMipMap& Mip = Target.Texture->PlatformData->Mips[0];
+	const unsigned char* PixelData = static_cast<const unsigned char*>(Mip.BulkData.Lock(LOCK_READ_ONLY));
+
+	uint32_t Width = static_cast<uint32_t>(Target.Texture->GetSurfaceWidth());
+	uint32_t Height = static_cast<uint32_t>(Target.Texture->GetSurfaceHeight());
+
+	if (Target.Texture->GetPixelFormat() == EPixelFormat::PF_B8G8R8A8)
+	{
+		const uint32_t NumComponents = Width * Height * 4;
+		RBGAPixelData.Reset(NumComponents);
+		RBGAPixelData.AddUninitialized(NumComponents);
+
+		for (uint32_t i = 0; i < NumComponents; i += 4)
+		{
+			RBGAPixelData[i + 0] = PixelData[i + 2];
+			RBGAPixelData[i + 1] = PixelData[i + 1];
+			RBGAPixelData[i + 2] = PixelData[i + 0];
+			RBGAPixelData[i + 3] = PixelData[i + 3];
+		}
+
+		PixelData = static_cast<const unsigned char*>(RBGAPixelData.GetData());
+	}
+
+	checkf(MLHandleIsValid(GetHandle()), TEXT("ImageTracker handle is invalid!"));
+	MLResult Result = MLImageTrackerAddTargetFromArray(GetHandle(), &Target.Settings, PixelData, Width, Height, MLImageTrackerImageFormat_RGBA, &Target.Handle);
+	Mip.BulkData.Unlock();
+
+	if (Result != MLResult_Ok)
+	{
+		UE_LOG(LogMagicLeapImageTracker, Error, TEXT("MLImageTrackerAddTargetFromArray failed with error '%s'!"), UTF8_TO_TCHAR(MLGetResultString(Result)));
+		TargetCreateTask.Type = FMagicLeapImageTrackerTask::EType::TargetCreateFailed;
+		PushCompletedTask(TargetCreateTask);
+		return;
+	}
+
+	// [3] Cache all the static data for this target.
+	checkf(MLHandleIsValid(GetHandle()), TEXT("ImageTracker weak pointer is invalid!"));
+	Result = MLImageTrackerGetTargetStaticData(GetHandle(), Target.Handle, &Target.Data);
+	if (Result != MLResult_Ok)
+	{
+		UE_LOG(LogMagicLeapImageTracker, Error, TEXT("MLImageTrackerGetTargetStaticData failed with error '%s'!"), UTF8_TO_TCHAR(MLGetResultString(Result)));
+		TargetCreateTask.Type = FMagicLeapImageTrackerTask::EType::TargetCreateFailed;
+		PushCompletedTask(TargetCreateTask);
+		return;
+	}
+
+	UE_LOG(LogMagicLeapImageTracker, Log, TEXT("SetTarget successfully set for %s"), *Target.Name);
+	TargetCreateTask.Type = FMagicLeapImageTrackerTask::EType::TargetCreateSucceeded;
+	TargetCreateTask.Target = Target;
+	PushCompletedTask(TargetCreateTask);
+	{
+		FScopeLock Lock(&TargetsMutex);
+		TrackedImageTargets.Add(Target.Name, Target);
+	}
+}
+
+void FMagicLeapImageTrackerRunnable::RemoveTarget(const FMagicLeapImageTrackerTarget& Target, bool bCanDestroyTracker)
+{
+	if (MLHandleIsValid(Target.Handle))
+	{
+		checkf(MLHandleIsValid(GetHandle()), TEXT("ImageTracker weak pointer is invalid!"));
+		MLResult Result = MLImageTrackerRemoveTarget(GetHandle(), Target.Handle);
+		UE_CLOG(Result != MLResult_Ok, LogMagicLeapImageTracker, Error, TEXT("MLImageTrackerRemoveTarget '%s' failed with error '%s'!"), *Target.Name, UTF8_TO_TCHAR(MLGetResultString(Result)));
+		TrackedImageTargets.Remove(Target.Name);
+
+		if (bCanDestroyTracker && TrackedImageTargets.Num() == 0)
+		{
+			DestroyTracker();
+		}
+	}
+}
+
+MLHandle FMagicLeapImageTrackerRunnable::GetHandle() const
+{
+	return FPlatformAtomics::AtomicRead(reinterpret_cast<const volatile int64*>(&ImageTracker));
+}
+#endif // WITH_MLSDK
diff --git a/Engine/Plugins/Lumin/MagicLeap/Source/MagicLeapImageTracker/Private/MagicLeapImageTrackerRunnable.h b/Engine/Plugins/Lumin/MagicLeap/Source/MagicLeapImageTracker/Private/MagicLeapImageTrackerRunnable.h
new file mode 100644
--- /dev/null
+++ b/Engine/Plugins/Lumin/MagicLeap/Source/MagicLeapImageTracker/Private/MagicLeapImageTrackerRunnable.h
@@ -0,0 +1,66 @@
+// Copyright 1998-2019 Epic Games, Inc. All Rights Reserved.
+
+#pragma once
+
+#include "Async/Async.h"
+#include "Engine/Texture2D.h"
+#include "IMagicLeapPlugin.h"
+#include "MagicLeapImageTrackerComponent.h"
+#include "MagicLeapRunnable.h"
+
+DECLARE_LOG_CATEGORY_EXTERN(LogMagicLeapImageTracker, Verbose, All);
+
+struct FMagicLeapImageTrackerTask : public FMagicLeapTask
+{
+	enum class EType : uint32
+	{
+		None,
+		TryCreateTarget,
+		TryRemoveTarget,
+		TargetCreateSucceeded,
+		TargetCreateFailed
+	};
+
+	EType Type;
+	FMagicLeapImageTrackerTarget Target;
+
+	FMagicLeapImageTrackerTask()
+	: Type(EType::None)
+	{}
+};
+
+class FMagicLeapImageTrackerRunnable : public FMagicLeapRunnable<FMagicLeapImageTrackerTask>
+{
+public:
+	FMagicLeapImageTrackerRunnable();
+
+	void Exit() override;
+	void Pause() override;
+	void Resume() override;
+	bool ProcessCurrentTask() override;
+
+	void SetTargetAsync(const FMagicLeapImageTrackerTarget& ImageTarget);
+	bool RemoveTargetAsync(const FString& InName);
+	uint32 GetMaxSimultaneousTargets() const;
+	void SetMaxSimultaneousTargets(uint32 NewNumTargets);
+	bool GetImageTrackerEnabled() const;
+	void SetImageTrackerEnabled(bool bEnabled);
+	void UpdateTargetsMainThread();
+	bool TryGetRelativeTransformMainThread(const FString& TargetName, FVector& OutLocation, FRotator& OutRotation);
+
+private:
+	TMap<FString, FMagicLeapImageTrackerTarget> TrackedImageTargets;
+	FCriticalSection TargetsMutex;
+	mutable FCriticalSection TrackerMutex;
+	TArray<uint8> RBGAPixelData;
+#if WITH_MLSDK
+	MLImageTrackerSettings Settings;
+	MLHandle ImageTracker;
+
+	bool CreateTracker();
+	void DestroyTracker();
+	void SetTarget();
+	void RemoveTarget(const FMagicLeapImageTrackerTarget& Target, bool bCanDestroyTracker = false);
+	MLHandle GetHandle() const;
+#endif // WITH_MLSDK
+};
diff --git a/Engine/Plugins/Lumin/MagicLeap/Source/MagicLeapImageTracker/Public/IMagicLeapImageTrackerModule.h b/Engine/Plugins/Lumin/MagicLeap/Source/MagicLeapImageTracker/Public/IMagicLeapImageTrackerModule.h
new file mode 100644
--- /dev/null
+++ b/Engine/Plugins/Lumin/MagicLeap/Source/MagicLeapImageTracker/Public/IMagicLeapImageTrackerModule.h
@@ -0,0 +1,35 @@
+// Copyright 1998-2019 Epic Games, Inc. All Rights Reserved.
+
+#pragma once
+
+#include "Modules/ModuleManager.h"
+
+/**
+ * The public interface to this module.  In most cases, this interface is only public to sibling modules
+ * within this plugin.
+ */
+class IMagicLeapImageTrackerModule : public IModuleInterface
+{
+public:
+
+	/**
+	 * Singleton-like access to this module's interface.  This is just for convenience!
+	 * Beware of calling this during the shutdown phase, though.  Your module might have been unloaded already.
+	 *
+	 * @return Returns singleton instance, loading the module on demand if needed
+	 */
+	static inline IMagicLeapImageTrackerModule& Get()
+	{
+		return FModuleManager::LoadModuleChecked<IMagicLeapImageTrackerModule>("MagicLeapImageTracker");
+	}
+
+	/**
+	 * Checks to see if this module is loaded and ready.  It is only valid to call Get() if IsAvailable() returns true.
+	 *
+	 * @return True if the module is loaded and ready to use
+	 */
+	static inline bool IsAvailable()
+	{
+		return FModuleManager::Get().IsModuleLoaded("MagicLeapImageTracker");
+	}
+};
diff --git a/Engine/Plugins/Lumin/MagicLeap/Source/MagicLeapImageTracker/Public/MagicLeapImageTrackerComponent.h b/Engine/Plugins/Lumin/MagicLeap/Source/MagicLeapImageTracker/Public/MagicLeapImageTrackerComponent.h
new file mode 100644
--- /dev/null
+++ b/Engine/Plugins/Lumin/MagicLeap/Source/MagicLeapImageTracker/Public/MagicLeapImageTrackerComponent.h
@@ -0,0 +1,116 @@
+// Copyright 1998-2019 Epic Games, Inc. All Rights Reserved.
+
+#pragma once
+
+#include "Components/SceneComponent.h"
+#include "Engine/Texture2D.h"
+#include "MagicLeapImageTrackerTypes.h"
+#include "MagicLeapImageTrackerComponent.generated.h"
+
+/**
+  The MLImageTrackerComponent will keep track of whether the image that it has been provided is currently
+  visible to the headset camera.
+  @note Currently only R8G8B8A8 and B8G8R8A8 textures are supported.
+*/
+UCLASS(ClassGroup = MagicLeap, BlueprintType, Blueprintable, EditInlineNew, meta = (BlueprintSpawnableComponent))
+class MAGICLEAPIMAGETRACKER_API UMagicLeapImageTrackerComponent : public USceneComponent
+{
+	GENERATED_BODY()
+
+public:
+	UMagicLeapImageTrackerComponent();
+
+	/**
+		Initiates the setting of the image target (if TargetImageTexture is not null).
+	*/
+	//void BeginPlay() override;
+
+	/**
+		Initiates the removal of the current image target (if valid).
+	*/
+	void EndPlay(const EEndPlayReason::Type EndPlayReason) override;
+
+	/**
+		Attempts to retrieve and set the relative pose of the tracked image.
+	*/
+	virtual void TickComponent(float DeltaTime, enum ELevelTick TickType, FActorComponentTickFunction *ThisTickFunction) override;
+
+	/**
+		Attempts to change the currently tracked target.  Initiates an asynchronous call on a worker thread.
+		When the task completes, the instigating blueprint will be notified by either a FSetImageTargetSucceeded
+		or FSetImageTargetFailed event.
+		@param ImageTarget The new texture to be tracked.
+		@return True if the initiation of the target change was successful, false otherwise.
+	*/
+	UFUNCTION(BlueprintCallable, Category = "ImageTracking|MagicLeap")
+	bool SetTargetAsync(UTexture2D* ImageTarget);
+
+	/**
+		Attempts to remove the currently tracked target.  Initiates an asynchronous call on a worker thread.
+		@return True if the initiation of the target removal was successful, false otherwise.
+	*/
+	UFUNCTION(BlueprintCallable, Category = "ImageTracking|MagicLeap")
+	bool RemoveTargetAsync();
+
+	/** The texture that will be tracked by this image tracker instance. */
+	UPROPERTY(EditAnywhere, BlueprintReadWrite, Category = "ImageTracking|MagicLeap")
+	UTexture2D* TargetImageTexture;
+
+	/**
+	  The name of the target.
+	  This name has to be unique across all instances of the ImageTrackerComponent class.
+	  If empty, the name of the component will be used.
+	*/
+	UPROPERTY(EditAnywhere, BlueprintReadWrite, Category = "ImageTracking|MagicLeap")
+	FString Name;
+
+	/** LongerDimension refers to the size of the longer dimension (width or height) of the physical image target in Unreal units. */
+	UPROPERTY(EditAnywhere, BlueprintReadWrite, Category = "ImageTracking|MagicLeap")
+	float LongerDimension;
+
+	/** Set this to true to improve detection for stationary targets. */
+	UPROPERTY(EditAnywhere, BlueprintReadWrite, Category = "ImageTracking|MagicLeap")
+	bool bIsStationary;
+
+	/** If false, the pose will not be updated when tracking is unreliable. */
+	UPROPERTY(EditAnywhere, BlueprintReadWrite, Category = "ImageTracking|MagicLeap")
+	bool bUseUnreliablePose;
+
+//private:
+	/** Activated when the target image is successfully set. */
+	UPROPERTY(BlueprintAssignable, Category = "ImageTracking | MagicLeap", meta = (AllowPrivateAccess = true))
+	FSetImageTargetSucceededMulti OnSetImageTargetSucceeded;
+
+	/** Activated when the target image fails to be set. */
+	UPROPERTY(BlueprintAssignable, Category = "ImageTracking | MagicLeap", meta = (AllowPrivateAccess = true))
+	FSetImageTargetFailedMulti OnSetImageTargetFailed;
+
+	/** Activated when the target image becomes visible to the camera */
+	UPROPERTY(BlueprintAssignable, Category = "ImageTracking | MagicLeap", meta = (AllowPrivateAccess = true))
+	FImageTargetFoundMulti OnImageTargetFound;
+
+	/** Activated the target image becomes invisible to the camera */
+	UPROPERTY(BlueprintAssignable, Category = "ImageTracking | MagicLeap", meta = (AllowPrivateAccess = true))
+	FImageTargetLostMulti OnImageTargetLost;
+
+	/**
+	  Activated when the target image is tracked with low confidence.
+
+	  The Image Tracker system will still provide a 6 DOF pose. But this
+	  pose might be inaccurate and might have jitter. When the tracking is
+	  unreliable one of the folling two events will happen quickly : Either
+	  the tracking will recover to Tracked or tracking will be lost and the
+	  status will change to NotTracked.
+	*/
+	UPROPERTY(BlueprintAssignable, Category = "ImageTracking | MagicLeap", meta = (AllowPrivateAccess = true))
+	FImageTargetUnreliableTrackingMulti OnImageTargetUnreliableTracking;
+
+private:
+	bool bIsTracking;
+#if WITH_EDITOR
+	UTexture2D* TextureBeforeEdit;
+public:
+	void PreEditChange(UProperty* PropertyAboutToChange) override;
+	void PostEditChangeProperty(FPropertyChangedEvent& PropertyChangedEvent) override;
+#endif
+};
diff --git a/Engine/Plugins/Lumin/MagicLeap/Source/MagicLeapImageTracker/Public/MagicLeapImageTrackerFunctionLibrary.h b/Engine/Plugins/Lumin/MagicLeap/Source/MagicLeapImageTracker/Public/MagicLeapImageTrackerFunctionLibrary.h
new file mode 100644
--- /dev/null
+++ b/Engine/Plugins/Lumin/MagicLeap/Source/MagicLeapImageTracker/Public/MagicLeapImageTrackerFunctionLibrary.h
@@ -0,0 +1,59 @@
+// Copyright 1998-2019 Epic Games, Inc. All Rights Reserved.
+
+#pragma once
+
+#include "CoreMinimal.h"
+#include "Kismet/BlueprintFunctionLibrary.h"
+#include "MagicLeapImageTrackerTypes.h"
+#include "MagicLeapImageTrackerFunctionLibrary.generated.h"
+
+UCLASS(ClassGroup = MagicLeap)
+class MAGICLEAPIMAGETRACKER_API UMagicLeapImageTrackerFunctionLibrary : public UBlueprintFunctionLibrary
+{
+	GENERATED_BODY()
+
+public:
+	/**
+	  Set maximum number of Image Targets that can be tracked at any given time.
+
+	  If the tracker is already tracking the maximum number of targets
+	  possible then it will stop searching for new targets which helps
+	  in reducing the load on the CPU. For example, if you are interested in
+	  tracking a maximum of x targets from a list y (x < y) targets then set this
+	  parameter to x.
+
+	  The valid range for this parameter is from 1 through 25.
+	*/
+	UFUNCTION(BlueprintCallable, Category = "ImageTracking|MagicLeap")
+	static void SetMaxSimultaneousTargets(int32 MaxSimultaneousTargets);
+
+	/**
+		Gets the maximum number of Image Targets that can be tracked at any given time.
+		@return The maximum number of Image Targets that can be tracked at any given time.
+	*/
+	UFUNCTION(BlueprintCallable, Category = "ImageTracking|MagicLeap")
+	static int32 GetMaxSimultaneousTargets();
+
+	/**
+	  If true, image tracker will detect and track targets.
+
+	  When enabled Image Tracker will gain access to the camera and start
+	  trackingimages. Enabling image tracker is expensive, takes about 1500ms
+	  on the average.
+
+	  When disabled Image Tracker will release the camera and stop tracking
+	  images. Internal state of the tracker will be maintained (i.e. list of
+	  active/inactive argets and their target_handles).
+
+	  This is done automatically on application pause / resume.
+	*/
+	UFUNCTION(BlueprintCallable, Category = "ImageTracking|MagicLeap")
+	static void EnableImageTracking(bool bEnable);
+
+	/**
+		Gets the active state of the image tracking system.
+		@return True if image tracking is enabled, false otherwise.
+	*/
+	UFUNCTION(BlueprintCallable, Category = "ImageTracking|MagicLeap")
+	static bool IsImageTrackingEnabled();
+};
diff --git a/Engine/Plugins/Lumin/MagicLeap/Source/MagicLeapImageTracker/Public/MagicLeapImageTrackerTypes.h b/Engine/Plugins/Lumin/MagicLeap/Source/MagicLeapImageTracker/Public/MagicLeapImageTrackerTypes.h
new file mode 100644
--- /dev/null
+++ b/Engine/Plugins/Lumin/MagicLeap/Source/MagicLeapImageTracker/Public/MagicLeapImageTrackerTypes.h
@@ -0,0 +1,78 @@
+// Copyright 1998-2019 Epic Games, Inc. All Rights Reserved.
+#pragma once
+
+//#include "CoreMinimal.h"
+#include "Engine/Engine.h"
+#include "Lumin/CAPIShims/LuminAPIImageTracking.h"
+#include "MagicLeapImageTrackerTypes.generated.h"
+
+/**
+	Delegate used to notify the instigating blueprint that the target image's location has become unrealiable.
+	@param LastTrackedLocation The last reliable location of the target image.
+	@param LastTrackedRotation The last reliable rotation of the target image.
+	@param NewUnreliableLocation The new location of the target image (which may or may not be accurate).
+	@param NewUnreliableRotation The new rotation of the target image (which may or may not be accurate).
+*/
+DECLARE_DYNAMIC_DELEGATE_FourParams(FImageTargetUnreliableTracking, const FVector&, LastTrackedLocation, const FRotator&, LastTrackedRotation, const FVector&, NewUnreliableLocation, const FRotator&, NewUnreliableRotation);
+DECLARE_DYNAMIC_MULTICAST_DELEGATE_FourParams(FImageTargetUnreliableTrackingMulti, const FVector&, LastTrackedLocation, const FRotator&, LastTrackedRotation, const FVector&, NewUnreliableLocation, const FRotator&, NewUnreliableRotation);
+
+/**
+	Delegate used to notify the instigating blueprint that the target image's location/rotation has changed.
+	@param NewLocation The new location of the target image (which may or may not be accurate).
+	@param NewRotation The new rotation of the target image (which may or may not be accurate).
+*/
+DECLARE_DYNAMIC_DELEGATE_TwoParams(FImageTargetReliableTracking, const FVector&, NewLocation, const FRotator&, NewRotation);
+DECLARE_DYNAMIC_MULTICAST_DELEGATE_TwoParams(FImageTargetReliableTrackingMulti, const FVector&, NewLocation, const FRotator&, NewRotation);
+
+/** Delegate used to notify the instigating blueprint that the target image was successfully set. */
+DECLARE_DYNAMIC_DELEGATE(FSetImageTargetSucceeded);
+DECLARE_DYNAMIC_MULTICAST_DELEGATE(FSetImageTargetSucceededMulti);
+
+/** Delegate used to notify the instigating blueprint that the target image failed to be set. */
+DECLARE_DYNAMIC_DELEGATE(FSetImageTargetFailed);
+DECLARE_DYNAMIC_MULTICAST_DELEGATE(FSetImageTargetFailedMulti);
+
+/** Delegate used to notify the instigating blueprint that the target image is currently visible to the camera */
+DECLARE_DYNAMIC_DELEGATE(FImageTargetFound);
+DECLARE_DYNAMIC_MULTICAST_DELEGATE(FImageTargetFoundMulti);
+
+/** Delegate used to notify the instigating blueprint that the target image just became invisible to the camera */
+DECLARE_DYNAMIC_DELEGATE(FImageTargetLost);
+DECLARE_DYNAMIC_MULTICAST_DELEGATE(FImageTargetLostMulti);
+
+USTRUCT()
+struct FMagicLeapImageTrackerTarget
+{
+	GENERATED_BODY()
+
+	FString Name;
+#if WITH_MLSDK
+	MLHandle Handle;
+	MLImageTrackerTargetStaticData Data;
+	MLImageTrackerTargetSettings Settings;
+	MLImageTrackerTargetResult OldTrackingStatus;
+#endif // WITH_MLSDK
+	UTexture2D* Texture;
+	FVector Location;
+	FRotator Rotation;
+	FVector UnreliableLocation;
+	FRotator UnreliableRotation;
+	bool bUseUnreliablePose;
+	FSetImageTargetSucceededMulti OnSetImageTargetSucceeded;
+	FSetImageTargetFailedMulti OnSetImageTargetFailed;
+	FImageTargetFoundMulti OnImageTargetFound;
+	FImageTargetLostMulti OnImageTargetLost;
+	FImageTargetUnreliableTrackingMulti OnImageTargetUnreliableTracking;
+
+	FMagicLeapImageTrackerTarget()
+	: Name(TEXT("Undefined"))
+#if WITH_MLSDK
+	, Handle(ML_INVALID_HANDLE)
+#endif // WITH_MLSDK
+	, Texture(nullptr)
+	{
+#if WITH_MLSDK
+		OldTrackingStatus.status = MLImageTrackerTargetStatus_Ensure32Bits;
+#endif // WITH_MLSDK
+	}
+};
diff --git a/Engine/Plugins/Lumin/MagicLeap/Source/MagicLeapPlanes/Source/MagicLeapPlanes.Build.cs b/Engine/Plugins/Lumin/MagicLeap/Source/MagicLeapPlanes/Source/MagicLeapPlanes.Build.cs
new file mode 100644
--- /dev/null
+++ b/Engine/Plugins/Lumin/MagicLeap/Source/MagicLeapPlanes/Source/MagicLeapPlanes.Build.cs
@@ -0,0 +1,33 @@
+// Copyright 1998-2019 Epic Games, Inc. All Rights Reserved.
+
+using System.IO;
+
+namespace UnrealBuildTool.Rules
+{
+	public class MagicLeapPlanes : ModuleRules
+	{
+		public MagicLeapPlanes( ReadOnlyTargetRules Target ) : base(Target)
+		{
+			PublicIncludePaths.Add(Path.Combine(ModuleDirectory, "Public"));
+
+			PublicDependencyModuleNames.AddRange(
+				new string[]
+				{
+					"Core",
+					"CoreUObject",
+					"Engine"
+				}
+			);
+
+			PrivateDependencyModuleNames.AddRange(
+				new string[]
+				{
+					"LuminRuntimeSettings",
+					"MLSDK",
+					"MagicLeap",
+					"HeadMountedDisplay",
+				}
+			);
+		}
+	}
+}
diff --git a/Engine/Plugins/Lumin/MagicLeap/Source/MagicLeapPlanes/Source/Private/MagicLeapPlanesComponent.cpp b/Engine/Plugins/Lumin/MagicLeap/Source/MagicLeapPlanes/Source/Private/MagicLeapPlanesComponent.cpp
new file mode 100644
--- /dev/null
+++ b/Engine/Plugins/Lumin/MagicLeap/Source/MagicLeapPlanes/Source/Private/MagicLeapPlanesComponent.cpp
@@ -0,0 +1,50 @@
+// Copyright 1998-2019 Epic Games, Inc. All Rights Reserved.
+
+#include "MagicLeapPlanesComponent.h"
+#include "MagicLeapPlanesModule.h"
+#include "Components/BoxComponent.h"
+
+UMagicLeapPlanesComponent::UMagicLeapPlanesComponent()
+: QueryFlags({ EMagicLeapPlaneQueryFlags::Vertical, EMagicLeapPlaneQueryFlags::Horizontal, EMagicLeapPlaneQueryFlags::Arbitrary, EMagicLeapPlaneQueryFlags::PreferInner })
+, MaxResults(10)
+, MinHolePerimeter(50.0f)
+, MinPlaneArea(400.0f)
+{
+	bAutoActivate = true;
+	SearchVolume = CreateDefaultSubobject<UBoxComponent>(TEXT("SearchVolume"));
+	SearchVolume->AttachToComponent(this, FAttachmentTransformRules::KeepRelativeTransform);
+	SearchVolume->SetCollisionEnabled(ECollisionEnabled::NoCollision);
+	SearchVolume->SetCanEverAffectNavigation(false);
+	SearchVolume->CanCharacterStepUpOn = ECanBeCharacterBase::ECB_No;
+	SearchVolume->SetCollisionObjectType(ECollisionChannel::ECC_WorldDynamic);
+	SearchVolume->SetGenerateOverlapEvents(false);
+	// Recommended default box extents for meshing - 10m (5m radius)
+	SearchVolume->SetBoxExtent(FVector(1000, 1000, 1000), false);
+}
+
+void UMagicLeapPlanesComponent::BeginPlay()
+{
+	Super::BeginPlay();
+	GetMagicLeapPlanesModule().CreateTracker();
+}
+
+void UMagicLeapPlanesComponent::EndPlay(const EEndPlayReason::Type EndPlayReason)
+{
+	GetMagicLeapPlanesModule().DestroyTracker();
+	Super::EndPlay(EndPlayReason);
+}
+
+bool UMagicLeapPlanesComponent::RequestPlanesAsync()
+{
+	FMagicLeapPlanesQuery QueryParams;
+	QueryParams.Flags = QueryFlags;
+	QueryParams.MaxResults = MaxResults;
+	QueryParams.MinHoleLength = MinHolePerimeter;
+	QueryParams.MinPlaneArea = MinPlaneArea;
+	QueryParams.SearchVolumePosition = SearchVolume->GetComponentLocation();
+	QueryParams.SearchVolumeOrientation = SearchVolume->GetComponentQuat();
+	QueryParams.SearchVolumeExtents = SearchVolume->GetScaledBoxExtent();
+	return GetMagicLeapPlanesModule().QueryBeginAsync(
+		QueryParams,
+		OnPlanesQueryResult);
+}
diff --git a/Engine/Plugins/Lumin/MagicLeap/Source/MagicLeapPlanes/Source/Private/MagicLeapPlanesFunctionLibrary.cpp b/Engine/Plugins/Lumin/MagicLeap/Source/MagicLeapPlanes/Source/Private/MagicLeapPlanesFunctionLibrary.cpp
new file mode 100644
--- /dev/null
+++ b/Engine/Plugins/Lumin/MagicLeap/Source/MagicLeapPlanes/Source/Private/MagicLeapPlanesFunctionLibrary.cpp
@@ -0,0 +1,89 @@
+// Copyright 1998-2019 Epic Games, Inc. All Rights Reserved.
+
+#include "MagicLeapPlanesFunctionLibrary.h"
+#include "MagicLeapPlanesModule.h"
+
+bool UMagicLeapPlanesFunctionLibrary::CreateTracker()
+{
+	return GetMagicLeapPlanesModule().CreateTracker();
+}
+
+bool UMagicLeapPlanesFunctionLibrary::DestroyTracker()
+{
+	return GetMagicLeapPlanesModule().DestroyTracker();
+}
+
+bool UMagicLeapPlanesFunctionLibrary::IsTrackerValid()
+{
+	return GetMagicLeapPlanesModule().IsTrackerValid();
+}
+
+bool UMagicLeapPlanesFunctionLibrary::PlanesQueryBeginAsync(const FMagicLeapPlanesQuery& InQuery, const FMagicLeapPlanesResultDelegate& InResultDelegate)
+{
+	FMagicLeapPlanesResultDelegateMulti ResultDelegate;
+	ResultDelegate.Add(InResultDelegate);
+
+	return GetMagicLeapPlanesModule().QueryBeginAsync(
+		InQuery,
+		ResultDelegate);
+}
+
+FTransform UMagicLeapPlanesFunctionLibrary::GetContentScale(const AActor* ContentActor, const FMagicLeapPlaneResult& PlaneResult)
+{
+	check(ContentActor != nullptr);
+
+	const FBox ContentActorBox = ContentActor->GetComponentsBoundingBox(true);
+
+	const FVector ContentActorScale = ContentActor->GetActorScale3D();
+	const FVector OriginalContentActorDimensions = (ContentActorBox.Max - ContentActor->GetActorLocation()) / ContentActorScale;
+
+	FVector FinalContentActorDimensions = FVector::ZeroVector;
+
+	const FVector2D PlaneDimensions = PlaneResult.PlaneDimensions / 2;
+
+	// For Unreal, Y = Right (Width)
+	FinalContentActorDimensions.Y = PlaneDimensions.Y;
+	FinalContentActorDimensions.X = FinalContentActorDimensions.Y * OriginalContentActorDimensions.X / OriginalContentActorDimensions.Y;
+	if (FinalContentActorDimensions.X > PlaneDimensions.X)
+	{
+		FinalContentActorDimensions.Y = FinalContentActorDimensions.Y * PlaneDimensions.X / FinalContentActorDimensions.X;
+		FinalContentActorDimensions.X = PlaneDimensions.X;
+	}
+	FinalContentActorDimensions.Z = OriginalContentActorDimensions.Z * FinalContentActorDimensions.X / OriginalContentActorDimensions.X;
+
+	const FVector FinalContentActorScale = FinalContentActorDimensions / OriginalContentActorDimensions;
+
+	return FTransform(PlaneResult.PlaneOrientation, PlaneResult.PlanePosition, FinalContentActorScale);
+}
+
+void UMagicLeapPlanesFunctionLibrary::ReorderPlaneFlags(const TArray<EMagicLeapPlaneQueryFlags>& InPriority, const TArray<EMagicLeapPlaneQueryFlags>& InFlagsToReorder, TArray<EMagicLeapPlaneQueryFlags>& OutReorderedFlags)
+{
+	struct FPlaneFlagPredicate
+	{
+		FPlaneFlagPredicate(const TArray<EMagicLeapPlaneQueryFlags>& Priority)
+		: PriorityArray(&Priority)
+		{}
+
+		bool operator()(const EMagicLeapPlaneQueryFlags& FlagA, const EMagicLeapPlaneQueryFlags& FlagB) const
+		{
+			return PriorityArray->Find(FlagA) > PriorityArray->Find(FlagB);
+		}
+	
+	private:
+		const TArray<EMagicLeapPlaneQueryFlags>* PriorityArray;
+	};
+
+	OutReorderedFlags = InFlagsToReorder;
+	OutReorderedFlags.HeapSort(FPlaneFlagPredicate(InPriority));
+}
+
+void UMagicLeapPlanesFunctionLibrary::RemoveFlagsNotInQuery(const TArray<EMagicLeapPlaneQueryFlags>& InQueryFlags, const TArray<EMagicLeapPlaneQueryFlags>& InResultFlags, TArray<EMagicLeapPlaneQueryFlags>& OutFlags)
+{
+	for (int32 i = 0; i < InResultFlags.Num(); ++i)
+	{
+		if (InQueryFlags.Find(InResultFlags[i]) >= 0)
+		{
+			OutFlags.Add(InResultFlags[i]);
+		}
+	}
+}
diff --git a/Engine/Plugins/Lumin/MagicLeap/Source/MagicLeapPlanes/Source/Private/MagicLeapPlanesModule.cpp b/Engine/Plugins/Lumin/MagicLeap/Source/MagicLeapPlanes/Source/Private/MagicLeapPlanesModule.cpp
new file mode 100644
--- /dev/null
+++ b/Engine/Plugins/Lumin/MagicLeap/Source/MagicLeapPlanes/Source/Private/MagicLeapPlanesModule.cpp
@@ -0,0 +1,395 @@
+// Copyright 1998-2019 Epic Games, Inc. All Rights Reserved.
+
+#include "MagicLeapPlanesModule.h"
+#include "MagicLeapMath.h"
+#include "HeadMountedDisplayFunctionLibrary.h"
+#include "Engine/Engine.h"
+#include "Components/BoxComponent.h"
+#include "MagicLeapHandle.h"
+
+//PRAGMA_DISABLE_OPTIMIZATION
+using namespace MagicLeap;
+
+DEFINE_LOG_CATEGORY(LogMagicLeapPlanes);
+
+#if WITH_MLSDK
+namespace MagicLeap
+{
+	EMagicLeapPlaneQueryFlags MLToUnrealPlanesQueryFlagMap(MLPlanesQueryFlags QueryFlag)
+	{
+		switch (QueryFlag)
+		{
+		case MLPlanesQueryFlag_Vertical:
+			return EMagicLeapPlaneQueryFlags::Vertical;
+		case MLPlanesQueryFlag_Horizontal:
+			return EMagicLeapPlaneQueryFlags::Horizontal;
+		case MLPlanesQueryFlag_Arbitrary:
+			return EMagicLeapPlaneQueryFlags::Arbitrary;
+		case MLPlanesQueryFlag_OrientToGravity:
+			return EMagicLeapPlaneQueryFlags::OrientToGravity;
+		case MLPlanesQueryFlag_Inner:
+			return EMagicLeapPlaneQueryFlags::PreferInner;
+		case MLPlanesQueryFlag_Semantic_Ceiling:
+			return EMagicLeapPlaneQueryFlags::Ceiling;
+		case MLPlanesQueryFlag_Semantic_Floor:
+			return EMagicLeapPlaneQueryFlags::Floor;
+		case MLPlanesQueryFlag_Semantic_Wall:
+			return EMagicLeapPlaneQueryFlags::Wall;
+		case MLPlanesQueryFlag_Polygons:
+			return EMagicLeapPlaneQueryFlags::Polygons;
+		}
+		return static_cast<EMagicLeapPlaneQueryFlags>(0);
+	}
+
+	void MLToUnrealPlanesQueryFlags(uint32 QueryFlags, TArray<EMagicLeapPlaneQueryFlags>& OutPlaneFlags)
+	{
+		static TArray<MLPlanesQueryFlags> AllMLFlags({
+		  MLPlanesQueryFlag_Vertical,
+		  MLPlanesQueryFlag_Horizontal,
+		  MLPlanesQueryFlag_Arbitrary,
+		  MLPlanesQueryFlag_OrientToGravity,
+		  MLPlanesQueryFlag_Inner,
+		  MLPlanesQueryFlag_Semantic_Ceiling,
+		  MLPlanesQueryFlag_Semantic_Floor,
+		  MLPlanesQueryFlag_Semantic_Wall,
+		  MLPlanesQueryFlag_Polygons
+		});
+
+		OutPlaneFlags.Empty();
+
+		for (MLPlanesQueryFlags flag : AllMLFlags)
+		{
+			if ((QueryFlags & static_cast<MLPlanesQueryFlags>(flag)) != 0)
+			{
+				OutPlaneFlags.Add(MLToUnrealPlanesQueryFlagMap(flag));
+			}
+		}
+	}
+
+	MLPlanesQueryFlags UnrealToMLPlanesQueryFlagMap(EMagicLeapPlaneQueryFlags QueryFlag)
+	{
+		switch (QueryFlag)
+		{
+		case EMagicLeapPlaneQueryFlags::Vertical:
+			return MLPlanesQueryFlag_Vertical;
+		case EMagicLeapPlaneQueryFlags::Horizontal:
+			return MLPlanesQueryFlag_Horizontal;
+		case EMagicLeapPlaneQueryFlags::Arbitrary:
+			return MLPlanesQueryFlag_Arbitrary;
+		case EMagicLeapPlaneQueryFlags::OrientToGravity:
+			return MLPlanesQueryFlag_OrientToGravity;
+		case EMagicLeapPlaneQueryFlags::PreferInner:
+			return MLPlanesQueryFlag_Inner;
+		case EMagicLeapPlaneQueryFlags::Ceiling:
+			return MLPlanesQueryFlag_Semantic_Ceiling;
+		case EMagicLeapPlaneQueryFlags::Floor:
+			return MLPlanesQueryFlag_Semantic_Floor;
+		case EMagicLeapPlaneQueryFlags::Wall:
+			return MLPlanesQueryFlag_Semantic_Wall;
+		case EMagicLeapPlaneQueryFlags::Polygons:
+			return MLPlanesQueryFlag_Polygons;
+		}
+		return static_cast<MLPlanesQueryFlags>(0);
+	}
+
+	MLPlanesQueryFlags UnrealToMLPlanesQueryFlags(const TArray<EMagicLeapPlaneQueryFlags>& QueryFlags)
+	{
+		MLPlanesQueryFlags MLFlags = static_cast<MLPlanesQueryFlags>(0);
+
+		for (EMagicLeapPlaneQueryFlags Flag : QueryFlags)
+		{
+			MLFlags = static_cast<MLPlanesQueryFlags>(MLFlags | UnrealToMLPlanesQueryFlagMap(Flag));
+		}
+
+		return MLFlags;
+	}
+}
+#endif // WITH_MLSDK
+
+FMagicLeapPlanesModule::FMagicLeapPlanesModule()
+{}
+
+void FMagicLeapPlanesModule::StartupModule()
+{
+	IMagicLeapPlanesModule::StartupModule();
+#if WITH_MLSDK
+	Tracker = ML_INVALID_HANDLE;
+#endif // WITH_MLSDK
+	TickDelegate = FTickerDelegate::CreateRaw(this, &FMagicLeapPlanesModule::Tick);
+
+	IMagicLeapPlugin::Get().RegisterMagicLeapTrackerEntity(this);
+}
+
+void FMagicLeapPlanesModule::ShutdownModule()
+{
+	IMagicLeapPlugin::Get().UnregisterMagicLeapTrackerEntity(this);
+
+	DestroyTracker();
+	IModuleInterface::ShutdownModule();
+}
+
+void FMagicLeapPlanesModule::DestroyEntityTracker()
+{
+	DestroyTracker();
+}
+
+bool FMagicLeapPlanesModule::Tick(float DeltaTime)
+{
+#if WITH_MLSDK
+	if (!(IMagicLeapPlugin::Get().IsMagicLeapHMDValid()))
+	{
+		return false;
+	}
+	else if (!MLHandleIsValid(Tracker))
+	{
+		return true;
+	}
+
+	const FTransform TrackingToWorld = UHeadMountedDisplayFunctionLibrary::GetTrackingToWorldTransform(nullptr);
+
+	for (int32 iRequest = PendingRequests.Num()-1; iRequest > -1; --iRequest)
+	{
+		FPlanesRequestMetaData& PendingRequest = PendingRequests[iRequest];
+		uint32 OutNumResults = 0;
+
+		MLPlaneBoundariesList BoundariesList;
+		MLPlaneBoundariesListInit(&BoundariesList);
+
+		MLResult Result = MLPlanesQueryGetResultsWithBoundaries(Tracker, PendingRequest.ResultHandle, PendingRequest.ResultMLPlanes.GetData(), &OutNumResults, &BoundariesList);
+		switch (Result)
+		{
+		case MLResult_Pending:
+		{
+			// Intentionally skip. We'll continue to check until it has completed.
+		}
+		break;
+		case MLResult_UnspecifiedFailure:
+		{
+			PendingRequest.ResultDelegateStatic.ExecuteIfBound(false, TArray<FMagicLeapPlaneResult>(), TArray<FMagicLeapPlaneBoundaries>());
+			PendingRequest.ResultDelegateDynamic.Broadcast(false, TArray<FMagicLeapPlaneResult>(), TArray<FMagicLeapPlaneBoundaries>());
+			PendingRequests.RemoveAt(iRequest);
+		}
+		break;
+		case MLResult_Ok:
+		{
+			const IMagicLeapPlugin& MLPlugin = IMagicLeapPlugin::Get();
+			float WorldToMetersScale = MLPlugin.GetWorldToMetersScale();
+
+			TArray<FMagicLeapPlaneResult> Planes;
+			Planes.Reserve(OutNumResults);
+
+			for (uint32 i = 0; i < OutNumResults; ++i)
+			{
+				FMagicLeapPlaneResult ResultPlane;
+				// Perception uses all coordinates in RUB so for them X axis is right and corresponds to the width of the plane.
+				// Unreal uses FRU, so the Y-axis is towards the right which makes the Y component of the vector the width.
+				ResultPlane.PlaneDimensions = FVector2D(PendingRequest.ResultMLPlanes[i].height * WorldToMetersScale, PendingRequest.ResultMLPlanes[i].width * WorldToMetersScale);
+
+				FTransform PlaneTransform = FTransform(MagicLeap::ToFQuat(PendingRequest.ResultMLPlanes[i].rotation), MagicLeap::ToFVector(PendingRequest.ResultMLPlanes[i].position, WorldToMetersScale), FVector(1.0f, 1.0f, 1.0f));
+				if (PlaneTransform.ContainsNaN())
+				{
+					UE_LOG(LogMagicLeapPlanes, Error, TEXT("Plane result %d transform contains NaN."), i);
+					continue;
+				}
+				if (!PlaneTransform.GetRotation().IsNormalized())
+				{
+					FQuat rotation = PlaneTransform.GetRotation();
+					rotation.Normalize();
+					PlaneTransform.SetRotation(rotation);
+				}
+
+				PlaneTransform.ConcatenateRotation(FQuat(FVector(0, 0, 1), PI));
+				PlaneTransform = PlaneTransform * TrackingToWorld;
+				ResultPlane.PlanePosition = PlaneTransform.GetLocation();
+				ResultPlane.PlaneOrientation = PlaneTransform.Rotator();
+				// The plane orientation has the forward axis (X) pointing in the direction of the plane's normal.
+				// We are rotating it by 90 degrees clock-wise about the right axis (Y) to get the up vector (Z) to point in the direction of the plane's normal.
+				// Since we are rotating the axis, the rotation is in the opposite direction of the object i.e. -90 degrees.
+				PlaneTransform.ConcatenateRotation(FQuat(FVector(0, 1, 0), -PI/2));
+				ResultPlane.ContentOrientation = PlaneTransform.Rotator();
+				ResultPlane.ID = MagicLeap::MLHandleToFGuid(PendingRequest.ResultMLPlanes[i].id);
+				MLToUnrealPlanesQueryFlags(PendingRequest.ResultMLPlanes[i].flags, ResultPlane.PlaneFlags);
+				
+				Planes.Add(ResultPlane);
+			}
+
+			TArray<FMagicLeapPlaneBoundaries> ResultBoundariesList;
+			if (BoundariesList.plane_boundaries != nullptr)
+			{
+				ResultBoundariesList.AddDefaulted(BoundariesList.plane_boundaries_count);
+				for (uint32 i = 0; i < BoundariesList.plane_boundaries_count; ++i)
+				{
+					const MLPlaneBoundaries& Boundaries = BoundariesList.plane_boundaries[i];
+
+					FMagicLeapPlaneBoundaries& ResultBoundaries = ResultBoundariesList[i];
+					ResultBoundaries.ID = MagicLeap::MLHandleToFGuid(Boundaries.id);
+					ResultBoundaries.Boundaries.AddDefaulted(Boundaries.boundaries_count);
+					for (uint32 j = 0; j < Boundaries.boundaries_count; ++j)
+					{
+						const MLPlaneBoundary& Boundary = Boundaries.boundaries[j];
+
+						FMagicLeapPlaneBoundary& ResultBoundary = ResultBoundaries.Boundaries[j];
+						ResultBoundary.Polygon.Vertices.AddDefaulted(Boundary.polygon->vertices_count);
+						for (uint32 k = 0; k < Boundary.polygon->vertices_count; ++k)
+						{
+							ResultBoundary.Polygon.Vertices[k] = MagicLeap::ToFVector(Boundary.polygon->vertices[k], WorldToMetersScale);
+							ResultBoundary.Polygon.Vertices[k] = TrackingToWorld.TransformPosition(ResultBoundary.Polygon.Vertices[k]);
+						}
+
+						if (Boundary.holes != nullptr && Boundary.holes_count > 0)
+						{
+							ResultBoundary.Holes.AddDefaulted(Boundary.holes_count);
+							const MLPolygon* hole = Boundary.holes;
+							for (uint32 k = 0; k < Boundary.holes_count; ++k)
+							{
+								ResultBoundary.Holes[k].Vertices.AddDefaulted(hole->vertices_count);
+								for (uint32 m = 0; m < hole->vertices_count; ++m)
+								{
+									ResultBoundary.Holes[k].Vertices[m] = MagicLeap::ToFVector(hole->vertices[m], WorldToMetersScale);
+									ResultBoundary.Holes[k].Vertices[m] = TrackingToWorld.TransformPosition(ResultBoundary.Holes[k].Vertices[m]);
+								}
+								++hole;
+							}
+						}
+					}
+				}
+
+				Result = MLPlanesReleaseBoundariesList(Tracker, &BoundariesList);
+				UE_CLOG(Result != MLResult_Ok, LogMagicLeapPlanes, Error, TEXT("MLPlanesReleaseBoundariesList failed with error %s"), UTF8_TO_TCHAR(MLGetResultString(Result)));
+			}
+
+			PendingRequest.ResultDelegateStatic.ExecuteIfBound(true, Planes, ResultBoundariesList);
+			PendingRequest.ResultDelegateDynamic.Broadcast(true, Planes, ResultBoundariesList);
+			PendingRequests.RemoveAt(iRequest);
+		}
+		break;
+		default:
+		{
+			UE_LOG(LogMagicLeapPlanes, Warning, TEXT("MLPlanesQueryGetResults failed with error '%s'"), UTF8_TO_TCHAR(MLGetResultString(Result)));
+		}
+		}
+	}
+#endif //WITH_MLSDK
+
+	return true;
+}
+
+bool FMagicLeapPlanesModule::CreateTracker()
+{
+#if WITH_MLSDK
+	if (!MLHandleIsValid(Tracker))
+	{
+		UE_LOG(LogMagicLeapPlanes, Display, TEXT("Creating Planes Tracker"));
+		MLResult Result = MLPlanesCreate(&Tracker);
+		UE_CLOG(Result != MLResult_Ok, LogMagicLeapPlanes, Error, TEXT("MLPlanesCreate failed with error '%s'."), UTF8_TO_TCHAR(MLGetResultString(Result)));
+
+		if (Result == MLResult_Ok)
+		{
+			TickDelegateHandle = FTicker::GetCoreTicker().AddTicker(TickDelegate);
+		}
+
+		return Result == MLResult_Ok;
+	}
+	else
+	{
+		UE_LOG(LogMagicLeapPlanes, Log, TEXT("Planes tracker already created."));
+	}
+#endif // WITH_MLSDK
+	return false;
+}
+
+bool FMagicLeapPlanesModule::DestroyTracker()
+{
+#if WITH_MLSDK
+	if (MLHandleIsValid(Tracker))
+	{
+		MLResult Result = MLPlanesDestroy(Tracker);
+		UE_CLOG(Result != MLResult_Ok, LogMagicLeapPlanes, Error, TEXT("MLPlanesDestroy failed with error '%s'."), UTF8_TO_TCHAR(MLGetResultString(Result)));
+		Tracker = ML_INVALID_HANDLE;
+		FTicker::GetCoreTicker().RemoveTicker(TickDelegateHandle);
+		return true;
+	}
+	else
+	{
+		UE_LOG(LogMagicLeapPlanes, Log, TEXT("Planes tracker already destroyed."));
+	}
+#endif // WITH_MLSDK
+	return false;
+}
+
+bool FMagicLeapPlanesModule::IsTrackerValid() const
+{
+#if WITH_MLSDK
+	return MLHandleIsValid(Tracker);
+#else
+	return false;
+#endif // WITH_MLSDK
+}
+
+bool FMagicLeapPlanesModule::QueryBeginAsync(const FMagicLeapPlanesQuery& QueryParams, const FMagicLeapPlanesResultStaticDelegate& InResultDelegate)
+{
+#if WITH_MLSDK
+	MLHandle Handle = SubmitPlanesQuery(QueryParams);
+	if (MLHandleIsValid(Handle))
+	{
+		FPlanesRequestMetaData& RequestMetaData = PendingRequests.AddDefaulted_GetRef();
+		RequestMetaData.ResultHandle = Handle;
+		RequestMetaData.ResultMLPlanes.AddDefaulted(QueryParams.MaxResults);
+		RequestMetaData.ResultDelegateStatic = InResultDelegate;
+		return true;
+	}
+#endif // WITH_MLSDK
+	return false;
+}
+
+bool FMagicLeapPlanesModule::QueryBeginAsync(
+	const FMagicLeapPlanesQuery& QueryParams,
+	const FMagicLeapPlanesResultDelegateMulti& InResultDelegate)
+{
+#if WITH_MLSDK
+	MLHandle Handle = SubmitPlanesQuery(QueryParams);
+	if (MLHandleIsValid(Handle))
+	{
+		FPlanesRequestMetaData& RequestMetaData = PendingRequests.AddDefaulted_GetRef();
+		RequestMetaData.ResultHandle = Handle;
+		RequestMetaData.ResultMLPlanes.AddDefaulted(QueryParams.MaxResults);
+		RequestMetaData.ResultDelegateDynamic = InResultDelegate;
+		return true;
+	}
+#endif // WITH_MLSDK
+	return false;
+}
+
+#if WITH_MLSDK
+MLHandle FMagicLeapPlanesModule::SubmitPlanesQuery(const FMagicLeapPlanesQuery& QueryParams)
+{
+	if (!IMagicLeapPlugin::Get().IsMagicLeapHMDValid())
+	{
+		return ML_INVALID_HANDLE;
+	}
+
+	const IMagicLeapPlugin& MLPlugin = IMagicLeapPlugin::Get();
+	float WorldToMetersScale = MLPlugin.GetWorldToMetersScale();
+	check(WorldToMetersScale != 0);
+
+	const FTransform WorldToTracking = UHeadMountedDisplayFunctionLibrary::GetTrackingToWorldTransform(nullptr).Inverse();
+
+	MLPlanesQuery Query;
+	Query.max_results = static_cast<uint32>(QueryParams.MaxResults);
+	Query.flags = UnrealToMLPlanesQueryFlags(QueryParams.Flags);
+	Query.min_hole_length = QueryParams.MinHoleLength / WorldToMetersScale;
+	Query.min_plane_area = QueryParams.MinPlaneArea / (WorldToMetersScale * WorldToMetersScale);
+	Query.bounds_center = MagicLeap::ToMLVector(WorldToTracking.TransformPosition(QueryParams.SearchVolumePosition), WorldToMetersScale);
+	Query.bounds_rotation = MagicLeap::ToMLQuat(WorldToTracking.TransformRotation(QueryParams.SearchVolumeOrientation));
+	// The C-API extents are 'full' extents - width, height, and depth. UE4 boxes are 'half' extents.
+	Query.bounds_extents = MagicLeap::ToMLVectorExtents(QueryParams.SearchVolumeExtents * 2, WorldToMetersScale);
+
+	MLHandle Handle = ML_INVALID_HANDLE;
+	MLResult Result = MLPlanesQueryBegin(Tracker, &Query, &Handle);
+	UE_CLOG(Result != MLResult_Ok, LogMagicLeapPlanes, Error, TEXT("MLPlanesQueryBegin failed with error '%s'."), UTF8_TO_TCHAR(MLGetResultString(Result)));
+
+	return Handle;
+}
+#endif // WITH_MLSDK
+
+IMPLEMENT_MODULE(FMagicLeapPlanesModule, MagicLeapPlanes);
diff --git a/Engine/Plugins/Lumin/MagicLeap/Source/MagicLeapPlanes/Source/Private/MagicLeapPlanesModule.h b/Engine/Plugins/Lumin/MagicLeap/Source/MagicLeapPlanes/Source/Private/MagicLeapPlanesModule.h
new file mode 100644
--- /dev/null
+++ b/Engine/Plugins/Lumin/MagicLeap/Source/MagicLeapPlanes/Source/Private/MagicLeapPlanesModule.h
@@ -0,0 +1,55 @@
+// Copyright 1998-2019 Epic Games, Inc. All Rights Reserved.
+
+#pragma once
+
+#include "CoreMinimal.h"
+#include "Containers/Ticker.h"
+#include "IMagicLeapPlugin.h"
+#include "IMagicLeapPlanesModule.h"
+#include "MagicLeapPlanesTypes.h"
+#include "IMagicLeapTrackerEntity.h"
+#include "Lumin/CAPIShims/LuminAPIPlanes.h"
+
+DECLARE_LOG_CATEGORY_EXTERN(LogMagicLeapPlanes, Verbose, All);
+
+class FMagicLeapPlanesModule : public IMagicLeapPlanesModule, public IMagicLeapTrackerEntity
+{
+public:
+	FMagicLeapPlanesModule();
+
+	void StartupModule() override;
+	void ShutdownModule() override;
+	bool Tick(float DeltaTime);
+
+	/** IMagicLeapTrackerEntity interface */
+	virtual void DestroyEntityTracker() override;
+
+	virtual bool CreateTracker() override;
+	virtual bool DestroyTracker() override;
+	virtual bool IsTrackerValid() const override;
+	virtual bool QueryBeginAsync(const FMagicLeapPlanesQuery& QueryParams, const FMagicLeapPlanesResultStaticDelegate& ResultDelegate) override;
+	virtual bool QueryBeginAsync(const FMagicLeapPlanesQuery& QueryParams, const FMagicLeapPlanesResultDelegateMulti& ResultDelegate) override;
+
+private:
+#if WITH_MLSDK
+	MLHandle SubmitPlanesQuery(const FMagicLeapPlanesQuery& QueryParams);
+
+	struct FPlanesRequestMetaData
+	{
+		MLHandle ResultHandle;
+		FMagicLeapPlanesResultDelegateMulti ResultDelegateDynamic;
+		FMagicLeapPlanesResultStaticDelegate ResultDelegateStatic;
+		TArray<MLPlane> ResultMLPlanes;
+	};
+
+	TArray<FPlanesRequestMetaData> PendingRequests;
+	MLHandle Tracker;
+#endif //WITH_MLSDK
+	FTickerDelegate TickDelegate;
+	FDelegateHandle TickDelegateHandle;
+};
+
+inline FMagicLeapPlanesModule& GetMagicLeapPlanesModule()
+{
+	return FModuleManager::Get().GetModuleChecked<FMagicLeapPlanesModule>("MagicLeapPlanes");
+}
diff --git a/Engine/Plugins/Lumin/MagicLeap/Source/MagicLeapPlanes/Source/Public/IMagicLeapPlanesModule.h b/Engine/Plugins/Lumin/MagicLeap/Source/MagicLeapPlanes/Source/Public/IMagicLeapPlanesModule.h
new file mode 100644
--- /dev/null
+++ b/Engine/Plugins/Lumin/MagicLeap/Source/MagicLeapPlanes/Source/Public/IMagicLeapPlanesModule.h
@@ -0,0 +1,51 @@
+// Copyright 1998-2019 Epic Games, Inc. All Rights Reserved.
+
+#pragma once
+
+#include "Modules/ModuleManager.h"
+#include "MagicLeapPlanesTypes.h"
+
+	/**
+	 * The public interface to this module.  In most cases, this interface is only public to sibling modules
+	 * within this plugin.
+	 */
+class IMagicLeapPlanesModule : public IModuleInterface
+{
+public:
+
+	/**
+	 * Singleton-like access to this module's interface.  This is just for convenience!
+	 * Beware of calling this during the shutdown phase, though.  Your module might have been unloaded already.
+	 *
+	 * @return Returns singleton instance, loading the module on demand if needed
+	 */
+	static inline IMagicLeapPlanesModule& Get()
+	{
+		return FModuleManager::LoadModuleChecked<IMagicLeapPlanesModule>("MagicLeapPlanes");
+	}
+
+	/**
+	 * Checks to see if this module is loaded and ready.  It is only valid to call Get() if IsAvailable() returns true.
+	 *
+	 * @return True if the module is loaded and ready to use
+	 */
+	static inline bool IsAvailable()
+	{
+		return FModuleManager::Get().IsModuleLoaded("MagicLeapPlanes");
+	}
+
+	/** Create a planes tracker. */
+	virtual bool CreateTracker() = 0;
+
+	/** Destroy a planes tracker. */
+	virtual bool DestroyTracker() = 0;
+
+	/** Is a planes tracker already created. */
+	virtual bool IsTrackerValid() const = 0;
+
+	/** Initiates a plane query with a static delegate. */
+	virtual bool QueryBeginAsync(const FMagicLeapPlanesQuery& Query, const FMagicLeapPlanesResultStaticDelegate& ResultDelegate) = 0;
+
+	/** Initiates a plane query with a dynamic delegate. */
+	virtual bool QueryBeginAsync(const FMagicLeapPlanesQuery& QueryParams, const FMagicLeapPlanesResultDelegateMulti& ResultDelegate) = 0;
+};
diff --git a/Engine/Plugins/Lumin/MagicLeap/Source/MagicLeapPlanes/Source/Public/MagicLeapPlanesComponent.h b/Engine/Plugins/Lumin/MagicLeap/Source/MagicLeapPlanes/Source/Public/MagicLeapPlanesComponent.h
new file mode 100644
--- /dev/null
+++ b/Engine/Plugins/Lumin/MagicLeap/Source/MagicLeapPlanes/Source/Public/MagicLeapPlanesComponent.h
@@ -0,0 +1,61 @@
+// Copyright 1998-2019 Epic Games, Inc. All Rights Reserved.
+
+#pragma once
+
+#include "Engine/Engine.h"
+#include "Components/ActorComponent.h"
+#include "MagicLeapPlanesTypes.h"
+#include "MagicLeapPlanesComponent.generated.h"
+
+/**
+	Component that provides access to the Planes API functionality.
+*/
+UCLASS(ClassGroup = MagicLeap, BlueprintType, Blueprintable, EditInlineNew, meta = (BlueprintSpawnableComponent))
+class MAGICLEAPPLANES_API UMagicLeapPlanesComponent : public USceneComponent
+{
+	GENERATED_BODY()
+
+public:
+	UMagicLeapPlanesComponent();
+
+	/** Creates the planes tracker handle for the component (or adds a ref count to an existing one). */
+	virtual void BeginPlay() override;
+
+	/** Removes a ref count to the underlying tracker (destroys the tracker if this is the last user). */
+	void EndPlay(const EEndPlayReason::Type EndPlayReason) override;
+
+	/** Initiates a plane query. */
+	UFUNCTION(BlueprintCallable, Category = "Planes | MagicLeap")
+	bool RequestPlanesAsync();
+
+	/** The flags to apply to this query. TODO: Should be a TSet but that is misbehaving in the editor. */
+	UPROPERTY(EditAnywhere, BlueprintReadWrite, Category = "Planes | MagicLeap")
+	TArray<EMagicLeapPlaneQueryFlags> QueryFlags;
+
+	/** Bounding box for searching planes in. */
+	UPROPERTY(EditAnywhere, BlueprintReadWrite, Category = "Planes | MagicLeap")
+	class UBoxComponent* SearchVolume;
+
+	/** The maximum number of planes that should be returned in the result. */
+	UPROPERTY(EditAnywhere, BlueprintReadWrite, Category = "Planes | MagicLeap", meta = (ClampMin = 0))
+	int32 MaxResults;
+
+	/**
+	  If EMagicLeapPlaneQueryFlags::IgnoreHoles is not a query flag then holes with a perimeter (in Unreal Units)
+	  smaller than this value will be ignored, and can be part of the plane.
+	*/
+	UPROPERTY(EditAnywhere, BlueprintReadWrite, Category = "Planes | MagicLeap", meta = (ClampMin = 0))
+	float MinHolePerimeter;
+
+	/**
+	  The minimum area (in squared Unreal Units) of planes to be returned.
+	  This value cannot be lower than 400 (lower values will be capped to this minimum).
+	*/
+	UPROPERTY(EditAnywhere, BlueprintReadWrite, Category = "Planes | MagicLeap", meta = (ClampMin = 400))
+	float MinPlaneArea;
+
+private:
+	// Delegate instances
+	UPROPERTY(BlueprintAssignable, Category = "Planes | MagicLeap", meta = (AllowPrivateAccess = true))
+	FMagicLeapPlanesResultDelegateMulti OnPlanesQueryResult;
+};
diff --git a/Engine/Plugins/Lumin/MagicLeap/Source/MagicLeapPlanes/Source/Public/MagicLeapPlanesFunctionLibrary.h b/Engine/Plugins/Lumin/MagicLeap/Source/MagicLeapPlanes/Source/Public/MagicLeapPlanesFunctionLibrary.h
new file mode 100644
--- /dev/null
+++ b/Engine/Plugins/Lumin/MagicLeap/Source/MagicLeapPlanes/Source/Public/MagicLeapPlanesFunctionLibrary.h
@@ -0,0 +1,59 @@
+// Copyright 1998-2019 Epic Games, Inc. All Rights Reserved.
+#pragma once
+
+#include "Engine/Engine.h"
+#include "Kismet/BlueprintFunctionLibrary.h"
+#include "MagicLeapPlanesTypes.h"
+#include "MagicLeapPlanesFunctionLibrary.generated.h"
+
+UCLASS(ClassGroup = MagicLeap)
+class MAGICLEAPPLANES_API UMagicLeapPlanesFunctionLibrary : public UBlueprintFunctionLibrary
+{
+	GENERATED_BODY()
+
+public:
+
+	/** Create a planes tracker. */
+	UFUNCTION(BlueprintCallable, Category = "Planes Function Library | MagicLeap")
+	static bool CreateTracker();
+
+	/** Destroy a planes tracker. */
+	UFUNCTION(BlueprintCallable, Category = "Planes Function Library | MagicLeap")
+	static bool DestroyTracker();
+
+	/** Is a planes tracker already created. */
+	UFUNCTION(BlueprintCallable, BlueprintPure, Category = "Planes Function Library | MagicLeap")
+	static bool IsTrackerValid();
+
+	/** Initiates a plane query. */
+	UFUNCTION(BlueprintCallable, Category = "Planes Function Library | MagicLeap")
+	static bool PlanesQueryBeginAsync(const FMagicLeapPlanesQuery& Query, const FMagicLeapPlanesResultDelegate& ResultDelegate);
+
+	/**
+		Gets the expected scale of the actor to be placed within the bounds of the given plane.
+		Ensure that the actor rotation is 0 (FQuat::Identity) before sending it to this function.
+		@param ContentActor Actor for which the scale is to be calculated.
+		@param PlaneDimensions Dimensions (in Unreal units) of the plane the actor has to be placed on.
+		@return A vector representing the expected scale of the Actor.
+	*/
+	UFUNCTION(BlueprintCallable, Category = "Planes Function Library | MagicLeap")
+	static FTransform GetContentScale(const AActor* ContentActor, const FMagicLeapPlaneResult& PlaneResult);
+
+	/**
+		Re-orders an array of plane query flags based on the priority list passed.
+		@param InPriority The priority list by which to order the array of plane query flags.
+		@param InFlagsToReorder The array of plane query flags to be reordered.
+		@param OutReorderedFlags The reordered array of plane query flags.
+	*/
+	UFUNCTION(BlueprintCallable, Category = "Planes Function Library | MagicLeap")
+	static void ReorderPlaneFlags(const TArray<EMagicLeapPlaneQueryFlags>& InPriority, const TArray<EMagicLeapPlaneQueryFlags>& InFlagsToReorder, TArray<EMagicLeapPlaneQueryFlags>& OutReorderedFlags);
+
+	/**
+		Removes 
+		@param InPriority The priority list by which to order the array of plane query flags.
+		@param InFlagsToReorder The array of plane query flags to be reordered.
+		@param OutReorderedFlags The reordered array of plane query flags.
+	*/
+	UFUNCTION(BlueprintCallable, Category = "Planes Function Library | MagicLeap")
+	static void RemoveFlagsNotInQuery(const TArray<EMagicLeapPlaneQueryFlags>& InQueryFlags, const TArray<EMagicLeapPlaneQueryFlags>& InResultFlags, TArray<EMagicLeapPlaneQueryFlags>& OutFlags);
+};
diff --git a/Engine/Plugins/Lumin/MagicLeap/Source/MagicLeapPlanes/Source/Public/MagicLeapPlanesTypes.h b/Engine/Plugins/Lumin/MagicLeap/Source/MagicLeapPlanes/Source/Public/MagicLeapPlanesTypes.h
new file mode 100644
--- /dev/null
+++ b/Engine/Plugins/Lumin/MagicLeap/Source/MagicLeapPlanes/Source/Public/MagicLeapPlanesTypes.h
@@ -0,0 +1,160 @@
+// Copyright 1998-2019 Epic Games, Inc. All Rights Reserved.
+#pragma once
+
+#include "Engine/Engine.h"
+#include "MagicLeapPlanesTypes.generated.h"
+
+/** Control flags for plane queries. */
+UENUM(BlueprintType)
+enum class EMagicLeapPlaneQueryFlags : uint8
+{
+	/** Include planes whose normal is perpendicular to gravity. */
+	Vertical,
+
+	/** Include planes whose normal is parallel to gravity. */
+	Horizontal,
+
+	/** Include planes with arbitrary normals. */
+	Arbitrary,
+
+	/** If set, non-horizontal planes will be aligned perpendicular to gravity. */
+	OrientToGravity,
+
+	/** If set, inner planes will be returned; if not set, outer planes will be returned. */
+	PreferInner,
+
+	/** If set, include planes semantically tagged as ceiling. */
+	Ceiling,
+
+	/** If set, include planes semantically tagged as floor. */
+	Floor,
+
+	/** If set, include planes semantically tagged as wall. */
+	Wall,
+
+	/** Include polygons */
+	Polygons
+};
+
+/** Represents a plane returned from the ML-API. */
+USTRUCT(BlueprintType)
+struct FMagicLeapPlaneResult
+{
+	GENERATED_BODY()
+
+public:
+	/** Position of the center of the plane in world coordinates. */
+	UPROPERTY(EditAnywhere, BlueprintReadWrite, Category = "Planes|MagicLeap")
+	FVector PlanePosition;
+
+	/** Orientation of the plane. */
+	UPROPERTY(EditAnywhere, BlueprintReadWrite, Category = "Planes|MagicLeap")
+	FRotator PlaneOrientation;
+
+	/** Orientation of the content with its up-vector orthogonal to the plane. */
+	UPROPERTY(EditAnywhere, BlueprintReadWrite, Category = "Planes|MagicLeap")
+	FRotator ContentOrientation;
+
+	/** Width and height of the plane (in Unreal units). */
+	UPROPERTY(EditAnywhere, BlueprintReadWrite, Category = "Planes|MagicLeap")
+	FVector2D PlaneDimensions;
+
+	/** The flags which describe this plane. TODO: Should be a TSet but that is misbehaving in the editor.*/
+	UPROPERTY(EditAnywhere, BlueprintReadWrite, Category = "Planes|MagicLeap")
+	TArray<EMagicLeapPlaneQueryFlags> PlaneFlags;
+
+	/** ID of the plane result. This ID is persistent across queries */
+	UPROPERTY(BlueprintReadOnly, Category = "Planes|MagicLeap")
+	FGuid ID;
+};
+
+/** Type used to represent a plane query. */
+USTRUCT(BlueprintType)
+struct FMagicLeapPlanesQuery
+{
+	GENERATED_BODY()
+
+	/** The flags to apply to this query. */
+	UPROPERTY(EditAnywhere, BlueprintReadWrite, Category = "Planes|MagicLeap")
+	TArray<EMagicLeapPlaneQueryFlags> Flags;
+
+	/** DEPRECATED. Use individual fields for setting search volume position, orientation and extents. */
+	UPROPERTY(EditAnywhere, BlueprintReadWrite, Category = "Planes|MagicLeap", meta = (DeprecatedProperty, DisplayName = "SearchVolume_DEPRECATED"))
+	class UBoxComponent* SearchVolume;
+
+	/**
+		The maximum number of results that should be returned.  This is also the minimum expected size of the array of results
+		passed to the MLPlanesGetResult function.
+	*/
+	UPROPERTY(EditAnywhere, BlueprintReadWrite, Category = "Planes|MagicLeap")
+	int32 MaxResults;
+
+	/**
+		If #MLPlanesQueryFlag_IgnoreHoles is set to false, holes with a perimeter (in meters) smaller than this value will be
+		ignored, and can be part of the plane.  This value cannot be lower than 0 (lower values will be capped to this minimum).
+		A good default value is 50cm.
+	*/
+	UPROPERTY(EditAnywhere, BlueprintReadWrite, Category = "Planes|MagicLeap")
+	float MinHoleLength;
+
+	/**
+		The minimum area (in squared meters) of planes to be returned.  This value cannot be lower than 400 (lower values
+		will be capped to this minimum).
+	*/
+	UPROPERTY(EditAnywhere, BlueprintReadWrite, Category = "Planes|MagicLeap")
+	float MinPlaneArea;
+
+	UPROPERTY(EditAnywhere, BlueprintReadWrite, Category = "Planes|MagicLeap")
+	FVector SearchVolumePosition;
+
+	UPROPERTY(EditAnywhere, BlueprintReadWrite, Category = "Planes|MagicLeap")
+	FQuat SearchVolumeOrientation;
+
+	UPROPERTY(EditAnywhere, BlueprintReadWrite, Category = "Planes|MagicLeap")
+	FVector SearchVolumeExtents;
+};
+
+USTRUCT(BlueprintType)
+struct FMagicLeapPolygon
+{
+	GENERATED_BODY()
+
+	/** The polygon that defines the region. */
+	UPROPERTY(EditAnywhere, BlueprintReadWrite, Category = "Planes|MagicLeap")
+	TArray<FVector> Vertices;
+};
+
+/** Coplanar connected line segments representing the outer boundary of an N-sided polygon. */
+USTRUCT(BlueprintType)
+struct FMagicLeapPlaneBoundary
+{
+	GENERATED_BODY()
+
+	/** The polygon that defines the region. */
+	UPROPERTY(EditAnywhere, BlueprintReadWrite, Category = "Planes|MagicLeap")
+	FMagicLeapPolygon Polygon;
+
+	/** A polygon may contains multiple holes. */
+	UPROPERTY(EditAnywhere, BlueprintReadWrite, Category = "Planes|MagicLeap")
+	TArray<FMagicLeapPolygon> Holes;
+};
+
+/** Represents multiple regions on a 2D plane. */
+USTRUCT(BlueprintType)
+struct FMagicLeapPlaneBoundaries
+{
+	GENERATED_BODY()
+
+	/** Plane ID, the same value associating to the ID in FMagicLeapPlaneResult if they belong to the same plane. */
+	UPROPERTY(EditAnywhere, BlueprintReadWrite, Category = "Planes|MagicLeap")
+	FGuid ID;
+
+	/** The polygon that defines the region. */
+	UPROPERTY(EditAnywhere, BlueprintReadWrite, Category = "Planes|MagicLeap")
+	TArray<FMagicLeapPlaneBoundary> Boundaries;
+};
+
+/** Delegate used to convey the result of a plane query. */
+DECLARE_DELEGATE_ThreeParams(FMagicLeapPlanesResultStaticDelegate, const bool, const TArray<FMagicLeapPlaneResult>&, const TArray<FMagicLeapPlaneBoundaries>&);
+DECLARE_DYNAMIC_DELEGATE_ThreeParams(FMagicLeapPlanesResultDelegate, const bool, bSuccess, const TArray<FMagicLeapPlaneResult>&, Planes, const TArray<FMagicLeapPlaneBoundaries>&, Polygons);
+DECLARE_DYNAMIC_MULTICAST_DELEGATE_ThreeParams(FMagicLeapPlanesResultDelegateMulti, const bool, bSuccess, const TArray<FMagicLeapPlaneResult>&, Planes, const TArray<FMagicLeapPlaneBoundaries>&, Polygons);
diff --git a/Engine/Plugins/Lumin/MagicLeap/Source/MagicLeapPrivileges/MagicLeapPrivileges.Build.cs b/Engine/Plugins/Lumin/MagicLeap/Source/MagicLeapPrivileges/MagicLeapPrivileges.Build.cs
new file mode 100644
--- /dev/null
+++ b/Engine/Plugins/Lumin/MagicLeap/Source/MagicLeapPrivileges/MagicLeapPrivileges.Build.cs
@@ -0,0 +1,22 @@
+// Copyright 1998-2019 Epic Games, Inc. All Rights Reserved.
+
+using System.IO;
+
+namespace UnrealBuildTool.Rules
+{
+	public class MagicLeapPrivileges : ModuleRules
+	{
+		public MagicLeapPrivileges(ReadOnlyTargetRules Target) : base(Target)
+		{
+			PrivateDependencyModuleNames.AddRange(
+				new string[]
+				{
+					"Core",
+					"CoreUObject",
+					"Engine",
+					"MLSDK"
+				}
+			);
+		}
+	}
+}
diff --git a/Engine/Plugins/Lumin/MagicLeap/Source/MagicLeapPrivileges/Private/MagicLeapPrivilegeUtils.cpp b/Engine/Plugins/Lumin/MagicLeap/Source/MagicLeapPrivileges/Private/MagicLeapPrivilegeUtils.cpp
new file mode 100644
--- /dev/null
+++ b/Engine/Plugins/Lumin/MagicLeap/Source/MagicLeapPrivileges/Private/MagicLeapPrivilegeUtils.cpp
@@ -0,0 +1,108 @@
+// Copyright 1998-2019 Epic Games, Inc. All Rights Reserved.
+
+#include "MagicLeapPrivilegeUtils.h"
+#include "MagicLeapPrivilegeTypes.h"
+#include "MagicLeapPrivilegesModule.h"
+
+#if WITH_MLSDK
+namespace MagicLeap
+{
+	MLPrivilegeID MAGICLEAPPRIVILEGES_API UnrealToMLPrivilege(EMagicLeapPrivilege Privilege)
+	{
+	#define PRIVCASE(x) case EMagicLeapPrivilege::x: { return MLPrivilegeID_##x; }
+		switch(Privilege)
+		{
+			PRIVCASE(BatteryInfo)
+			PRIVCASE(CameraCapture)
+			PRIVCASE(WorldReconstruction)
+			PRIVCASE(InAppPurchase)
+			PRIVCASE(AudioCaptureMic)
+			PRIVCASE(DrmCertificates)
+			PRIVCASE(Occlusion)
+			PRIVCASE(LowLatencyLightwear)
+			PRIVCASE(Internet)
+			PRIVCASE(IdentityRead)
+			PRIVCASE(BackgroundDownload)
+			PRIVCASE(BackgroundUpload)
+			PRIVCASE(MediaDrm)
+			PRIVCASE(Media)
+			PRIVCASE(MediaMetadata)
+			PRIVCASE(PowerInfo)
+			PRIVCASE(LocalAreaNetwork)
+			PRIVCASE(VoiceInput)
+			PRIVCASE(Documents)
+			PRIVCASE(ConnectBackgroundMusicService)
+			PRIVCASE(RegisterBackgroundMusicService)
+			PRIVCASE(PwFoundObjRead)
+			PRIVCASE(NormalNotificationsUsage)
+			PRIVCASE(MusicService)
+			PRIVCASE(ControllerPose)
+			PRIVCASE(ScreensProvider)
+			PRIVCASE(GesturesSubscribe)
+			PRIVCASE(GesturesConfig)
+			PRIVCASE(AddressBookRead)
+			PRIVCASE(AddressBookWrite)
+			PRIVCASE(CoarseLocation)
+			PRIVCASE(HandMesh)
+			PRIVCASE(WifiStatusRead)
+		default:
+			UE_LOG(LogMagicLeapPrivileges, Error, TEXT("Unmapped privilege %d"), static_cast<int32>(Privilege));
+			break;
+		}
+		return MLPrivilegeID_Invalid;
+	}
+
+	FString MAGICLEAPPRIVILEGES_API MLPrivilegeToString(MLPrivilegeID PrivilegeID)
+	{
+		#define PRIV_TO_STR_CASE(x) case x: { return UTF8_TO_TCHAR((#x)); }
+		switch (PrivilegeID)
+		{
+		PRIV_TO_STR_CASE(MLPrivilegeID_Invalid)
+		PRIV_TO_STR_CASE(MLPrivilegeID_BatteryInfo)
+		PRIV_TO_STR_CASE(MLPrivilegeID_CameraCapture)
+		PRIV_TO_STR_CASE(MLPrivilegeID_WorldReconstruction)
+		PRIV_TO_STR_CASE(MLPrivilegeID_InAppPurchase)
+		PRIV_TO_STR_CASE(MLPrivilegeID_AudioCaptureMic)
+		PRIV_TO_STR_CASE(MLPrivilegeID_DrmCertificates)
+		PRIV_TO_STR_CASE(MLPrivilegeID_Occlusion)
+		PRIV_TO_STR_CASE(MLPrivilegeID_LowLatencyLightwear)
+		PRIV_TO_STR_CASE(MLPrivilegeID_Internet)
+		PRIV_TO_STR_CASE(MLPrivilegeID_IdentityRead)
+		PRIV_TO_STR_CASE(MLPrivilegeID_BackgroundDownload)
+		PRIV_TO_STR_CASE(MLPrivilegeID_BackgroundUpload)
+		PRIV_TO_STR_CASE(MLPrivilegeID_MediaDrm)
+		PRIV_TO_STR_CASE(MLPrivilegeID_Media)
+		PRIV_TO_STR_CASE(MLPrivilegeID_MediaMetadata)
+		PRIV_TO_STR_CASE(MLPrivilegeID_PowerInfo)
+		PRIV_TO_STR_CASE(MLPrivilegeID_LocalAreaNetwork)
+		PRIV_TO_STR_CASE(MLPrivilegeID_VoiceInput)
+		PRIV_TO_STR_CASE(MLPrivilegeID_Documents)
+		PRIV_TO_STR_CASE(MLPrivilegeID_ConnectBackgroundMusicService)
+		PRIV_TO_STR_CASE(MLPrivilegeID_RegisterBackgroundMusicService)
+		PRIV_TO_STR_CASE(MLPrivilegeID_PwFoundObjRead)
+		PRIV_TO_STR_CASE(MLPrivilegeID_NormalNotificationsUsage)
+		PRIV_TO_STR_CASE(MLPrivilegeID_MusicService)
+		PRIV_TO_STR_CASE(MLPrivilegeID_ControllerPose)
+		PRIV_TO_STR_CASE(MLPrivilegeID_ScreensProvider)
+		PRIV_TO_STR_CASE(MLPrivilegeID_GesturesSubscribe)
+		PRIV_TO_STR_CASE(MLPrivilegeID_GesturesConfig)
+		PRIV_TO_STR_CASE(MLPrivilegeID_AddressBookRead)
+		PRIV_TO_STR_CASE(MLPrivilegeID_AddressBookWrite)
+		PRIV_TO_STR_CASE(MLPrivilegeID_CoarseLocation)
+		// TODO: @njain uncomment after enum is added to c-api
+		// PRIV_TO_STR_CASE(MLPrivilegeID_HandMesh)
+		PRIV_TO_STR_CASE(MLPrivilegeID_WifiStatusRead)
+		default:
+			UE_LOG(LogMagicLeapPrivileges, Error, TEXT("Unmapped privilege %d"), static_cast<int32>(PrivilegeID));
+			break;
+		}
+
+		return UTF8_TO_TCHAR("");
+	}
+
+	FString MAGICLEAPPRIVILEGES_API MLPrivilegeToString(EMagicLeapPrivilege PrivilegeID)
+	{
+		return MLPrivilegeToString(UnrealToMLPrivilege(PrivilegeID));
+	}
+}
+#endif //WITH_MLSDK
diff --git a/Engine/Plugins/Lumin/MagicLeap/Source/MagicLeapPrivileges/Private/MagicLeapPrivilegesFunctionLibrary.cpp b/Engine/Plugins/Lumin/MagicLeap/Source/MagicLeapPrivileges/Private/MagicLeapPrivilegesFunctionLibrary.cpp
new file mode 100644
--- /dev/null
+++ b/Engine/Plugins/Lumin/MagicLeap/Source/MagicLeapPrivileges/Private/MagicLeapPrivilegesFunctionLibrary.cpp
@@ -0,0 +1,30 @@
+// Copyright 1998-2019 Epic Games, Inc. All Rights Reserved.
+
+#include "MagicLeapPrivilegesFunctionLibrary.h"
+#include "MagicLeapPrivilegesModule.h"
+#include "MagicLeapPrivilegeUtils.h"
+
+bool UMagicLeapPrivilegesFunctionLibrary::CheckPrivilege(EMagicLeapPrivilege Privilege)
+{
+#if WITH_MLSDK
+	return GetMagicLeapPrivilegesModule().CheckPrivilege(Privilege);
+#endif //WITH_MLSDK
+	return false;
+}
+
+bool UMagicLeapPrivilegesFunctionLibrary::RequestPrivilege(EMagicLeapPrivilege Privilege)
+{
+#if WITH_MLSDK
+	return GetMagicLeapPrivilegesModule().RequestPrivilege(Privilege);
+#endif //WITH_MLSDK
+	return false;
+}
+
+bool UMagicLeapPrivilegesFunctionLibrary::RequestPrivilegeAsync(EMagicLeapPrivilege Privilege, const FMagicLeapPrivilegeRequestDelegate& ResultDelegate)
+{
+#if WITH_MLSDK
+	return GetMagicLeapPrivilegesModule().RequestPrivilegeAsync(Privilege, ResultDelegate);
+#else
+	return false;
+#endif //WITH_MLSDK
+}
diff --git a/Engine/Plugins/Lumin/MagicLeap/Source/MagicLeapPrivileges/Private/MagicLeapPrivilegesModule.cpp b/Engine/Plugins/Lumin/MagicLeap/Source/MagicLeapPrivileges/Private/MagicLeapPrivilegesModule.cpp
new file mode 100644
--- /dev/null
+++ b/Engine/Plugins/Lumin/MagicLeap/Source/MagicLeapPrivileges/Private/MagicLeapPrivilegesModule.cpp
@@ -0,0 +1,124 @@
+// Copyright 1998-2019 Epic Games, Inc. All Rights Reserved.
+
+#include "MagicLeapPrivilegesModule.h"
+#include "MagicLeapPrivilegeUtils.h"
+#include "Engine/Engine.h"
+
+using namespace MagicLeap;
+
+DEFINE_LOG_CATEGORY(LogMagicLeapPrivileges);
+
+FMagicLeapPrivilegesModule::FMagicLeapPrivilegesModule()
+{}
+
+void FMagicLeapPrivilegesModule::StartupModule()
+{
+	IMagicLeapPrivilegesModule::StartupModule();
+	TickDelegate = FTickerDelegate::CreateRaw(this, &FMagicLeapPrivilegesModule::Tick);
+	TickDelegateHandle = FTicker::GetCoreTicker().AddTicker(TickDelegate);
+}
+
+void FMagicLeapPrivilegesModule::ShutdownModule()
+{
+	FTicker::GetCoreTicker().RemoveTicker(TickDelegateHandle);
+	IModuleInterface::ShutdownModule();
+}
+
+bool FMagicLeapPrivilegesModule::Tick(float DeltaTime)
+{
+#if WITH_MLSDK
+	auto CopyPendingAsyncRequests(PendingAsyncRequests);
+
+	PendingAsyncRequests.Empty();
+
+	for (const auto& PendingAsyncRequest : CopyPendingAsyncRequests)
+	{
+		// Re-add pending ones
+		MLResult Result = MLPrivilegesRequestPrivilegeTryGet(PendingAsyncRequest.Request);
+		if (MLResult_Pending == Result)
+		{
+			PendingAsyncRequests.Add(PendingAsyncRequest);
+			continue;
+		}
+
+		// Dispatch. Granted gets true, everything else gets false.
+		if (MLPrivilegesResult_Granted == Result)
+		{
+			PendingAsyncRequest.StaticDelegate.ExecuteIfBound(PendingAsyncRequest.Privilege, true);
+			PendingAsyncRequest.DynamicDelegate.ExecuteIfBound(PendingAsyncRequest.Privilege, true);
+		}
+		else
+		{
+			PendingAsyncRequest.StaticDelegate.ExecuteIfBound(PendingAsyncRequest.Privilege, false);
+			PendingAsyncRequest.DynamicDelegate.ExecuteIfBound(PendingAsyncRequest.Privilege, false);
+		}
+	}
+#endif //WITH_MLSDK
+	return true;
+}
+
+bool FMagicLeapPrivilegesModule::CheckPrivilege(EMagicLeapPrivilege Privilege)
+{
+#if WITH_MLSDK
+	auto Result = MLPrivilegesCheckPrivilege(MagicLeap::UnrealToMLPrivilege(Privilege));
+	UE_CLOG(Result == MLResult_Ok, LogMagicLeapPrivileges, Error,
+		TEXT("MLPrivilegesCheckPrivilege('%s') failed with error '%s'."),
+		*MLPrivilegeToString(Privilege), UTF8_TO_TCHAR(MLPrivilegesGetResultString(Result)));
+	if (MLPrivilegesResult_Granted == Result)
+	{
+		return true;
+	}
+#endif //WITH_MLSDK
+	return false;
+}
+
+bool FMagicLeapPrivilegesModule::RequestPrivilege(EMagicLeapPrivilege Privilege)
+{
+#if WITH_MLSDK
+	auto Result = MLPrivilegesRequestPrivilege(MagicLeap::UnrealToMLPrivilege(Privilege));
+	UE_CLOG(Result == MLResult_Ok, LogMagicLeapPrivileges, Error,
+		TEXT("MLPrivilegesRequestPrivilege('%s') failed with error '%s'."),
+		*MLPrivilegeToString(Privilege), UTF8_TO_TCHAR(MLPrivilegesGetResultString(Result)));
+	if (MLPrivilegesResult_Granted == Result)
+	{
+		return true;
+	}
+#endif //WITH_MLSDK
+	return false;
+}
+
+bool FMagicLeapPrivilegesModule::RequestPrivilegeAsync(EMagicLeapPrivilege Privilege, const FMagicLeapPrivilegeRequestStaticDelegate& ResultDelegate)
+{
+	return RequestPrivilegeAsync(Privilege, ResultDelegate, FMagicLeapPrivilegeRequestDelegate());
+}
+
+bool FMagicLeapPrivilegesModule::RequestPrivilegeAsync(EMagicLeapPrivilege Privilege, const FMagicLeapPrivilegeRequestDelegate& ResultDelegate)
+{
+	return RequestPrivilegeAsync(Privilege, FMagicLeapPrivilegeRequestStaticDelegate(), ResultDelegate);
+}
+
+bool FMagicLeapPrivilegesModule::RequestPrivilegeAsync(EMagicLeapPrivilege Privilege, const FMagicLeapPrivilegeRequestStaticDelegate& StaticResultDelegate, const FMagicLeapPrivilegeRequestDelegate& DynamicResultDelegate)
+{
+#if WITH_MLSDK
+	MLPrivilegesAsyncRequest* AsyncPrivilegeRequest = nullptr;
+
+	auto Result = MLPrivilegesRequestPrivilegeAsync(MagicLeap::UnrealToMLPrivilege(Privilege), &AsyncPrivilegeRequest);
+
+	if (MLResult_Ok != Result)
+	{
+		UE_LOG(LogMagicLeapPrivileges, Error,
+			TEXT("MLPrivilegesRequestPrivilegeAsync('%s') failed with error '%s'."),
+			*MLPrivilegeToString(Privilege), UTF8_TO_TCHAR(MLPrivilegesGetResultString(Result)));
+		return false;
+	}
+
+	// Store the request
+	PendingAsyncRequests.Add({ Privilege, AsyncPrivilegeRequest, StaticResultDelegate, DynamicResultDelegate });
+	return true;
+#else
+	return false;
+#endif //WITH_MLSDK
+}
+
+
+IMPLEMENT_MODULE(FMagicLeapPrivilegesModule, MagicLeapPrivileges);
diff --git a/Engine/Plugins/Lumin/MagicLeap/Source/MagicLeapPrivileges/Private/MagicLeapPrivilegesModule.h b/Engine/Plugins/Lumin/MagicLeap/Source/MagicLeapPrivileges/Private/MagicLeapPrivilegesModule.h
new file mode 100644
--- /dev/null
+++ b/Engine/Plugins/Lumin/MagicLeap/Source/MagicLeapPrivileges/Private/MagicLeapPrivilegesModule.h
@@ -0,0 +1,46 @@
+// Copyright 1998-2019 Epic Games, Inc. All Rights Reserved.
+
+#pragma once
+
+#include "CoreMinimal.h"
+#include "Containers/Ticker.h"
+#include "IMagicLeapPrivilegesModule.h"
+#include "Lumin/CAPIShims/LuminAPI.h"
+
+DECLARE_LOG_CATEGORY_EXTERN(LogMagicLeapPrivileges, Verbose, All);
+
+class FMagicLeapPrivilegesModule : public IMagicLeapPrivilegesModule
+{
+public:
+	FMagicLeapPrivilegesModule();
+
+	void StartupModule() override;
+	void ShutdownModule() override;
+	bool Tick(float DeltaTime);
+
+	bool RequestPrivilege(EMagicLeapPrivilege Privilege);
+	bool CheckPrivilege(EMagicLeapPrivilege Privilege);
+	bool RequestPrivilegeAsync(EMagicLeapPrivilege Privilege, const FMagicLeapPrivilegeRequestStaticDelegate& ResultDelegate) override;
+	bool RequestPrivilegeAsync(EMagicLeapPrivilege Privilege, const FMagicLeapPrivilegeRequestDelegate& ResultDelegate) override;
+
+private:
+#if WITH_MLSDK
+	struct FPendingAsyncRequest
+	{
+		EMagicLeapPrivilege Privilege;
+		struct MLPrivilegesAsyncRequest* Request;
+		FMagicLeapPrivilegeRequestStaticDelegate StaticDelegate;
+		FMagicLeapPrivilegeRequestDelegate DynamicDelegate;
+	};
+	TArray<FPendingAsyncRequest> PendingAsyncRequests;
+#endif //WITH_MLSDK
+	FTickerDelegate TickDelegate;
+	FDelegateHandle TickDelegateHandle;
+
+	bool RequestPrivilegeAsync(EMagicLeapPrivilege Privilege, const FMagicLeapPrivilegeRequestStaticDelegate& StaticResultDelegate, const FMagicLeapPrivilegeRequestDelegate& DynamicResultDelegate);
+};
+
+inline FMagicLeapPrivilegesModule& GetMagicLeapPrivilegesModule()
+{
+	return FModuleManager::Get().GetModuleChecked<FMagicLeapPrivilegesModule>("MagicLeapPrivileges");
+}
