Adding includes

diff --git a/Engine/Plugins/Lumin/MagicLeap/Source/MagicLeapAudio/Private/FakeDeviceCallbackRunnable.cpp b/Engine/Plugins/Lumin/MagicLeap/Source/MagicLeapAudio/Private/FakeDeviceCallbackRunnable.cpp
--- a/Engine/Plugins/Lumin/MagicLeap/Source/MagicLeapAudio/Private/FakeDeviceCallbackRunnable.cpp
+++ b/Engine/Plugins/Lumin/MagicLeap/Source/MagicLeapAudio/Private/FakeDeviceCallbackRunnable.cpp
@@ -2,6 +2,8 @@
 
 #include "FakeDeviceCallbackRunnable.h"
 #include "CoreGlobals.h"
+#include "HAL/Event.h"
+#include "GenericPlatform/GenericPlatformProcess.h"
 
 namespace Audio
 {
