Missing include.

diff --git a/Engine/Plugins/Lumin/MagicLeap/Source/MagicLeap/Private/MagicLeapHMD.cpp b/Engine/Plugins/Lumin/MagicLeap/Source/MagicLeap/Private/MagicLeapHMD.cpp
--- a/Engine/Plugins/Lumin/MagicLeap/Source/MagicLeap/Private/MagicLeapHMD.cpp
+++ b/Engine/Plugins/Lumin/MagicLeap/Source/MagicLeap/Private/MagicLeapHMD.cpp
@@ -56,7 +56,7 @@
 #endif
 
 #if WITH_EDITOR
-#include "Editor/EditorEngine.h"
+#include "Editor.h"
 #include "Editor/EditorPerformanceSettings.h"
 #include "Misc/CoreDelegates.h"
 #include "ISettingsModule.h"
