Fixing SA warnings.

diff --git a/Engine/Plugins/Lumin/MagicLeap/Source/MagicLeapAR/Private/LuminARAPI.cpp b/Engine/Plugins/Lumin/MagicLeap/Source/MagicLeapAR/Private/LuminARAPI.cpp
--- a/Engine/Plugins/Lumin/MagicLeap/Source/MagicLeapAR/Private/LuminARAPI.cpp
+++ b/Engine/Plugins/Lumin/MagicLeap/Source/MagicLeapAR/Private/LuminARAPI.cpp
@@ -144,7 +144,6 @@ uint32 FLuminARSession::GetFrameNum()
 // Anchors and Planes
 ELuminARAPIStatus FLuminARSession::CreateARAnchor(const FTransform& TransfromInTrackingSpace, UARTrackedGeometry* TrackedGeometry, USceneComponent* ComponentToPin, FName InDebugName, UARPin*& OutAnchor)
 {
-	ELuminARAPIStatus AnchorCreateStatus = ELuminARAPIStatus::AR_SUCCESS;
 	OutAnchor = nullptr;
 
 	TSharedPtr<LuminArAnchor> NewLuminArAnchor;
@@ -167,17 +166,15 @@ ELuminARAPIStatus FLuminARSession::CreateARAnchor(const FTransform& TransfromInT
 		NewLuminArAnchor = MakeShared<LuminArAnchor>(Pose, ParentHandle);
 	}
 
-	if (AnchorCreateStatus == ELuminARAPIStatus::AR_SUCCESS)
-	{
-		OutAnchor = NewObject<UARPin>();
-		OutAnchor->InitARPin(GetARSystem(), ComponentToPin, TransfromInTrackingSpace, TrackedGeometry, InDebugName);
-		OutAnchor->SetNativeResource(reinterpret_cast<void*>(NewLuminArAnchor.Get()));
+	OutAnchor = NewObject<UARPin>();
+	OutAnchor->InitARPin(GetARSystem(), ComponentToPin, TransfromInTrackingSpace, TrackedGeometry, InDebugName);
+	OutAnchor->SetNativeResource(reinterpret_cast<void*>(NewLuminArAnchor.Get()));
 
-		UObjectManager->HandleToLuminAnchorMap.Add(NewLuminArAnchor->Handle, NewLuminArAnchor);
-		UObjectManager->AllAnchors.Add(OutAnchor);
-		UObjectManager->HandleToAnchorMap.Add(NewLuminArAnchor->Handle, OutAnchor);
-	}
-	return AnchorCreateStatus;
+	UObjectManager->HandleToLuminAnchorMap.Add(NewLuminArAnchor->Handle, NewLuminArAnchor);
+	UObjectManager->AllAnchors.Add(OutAnchor);
+	UObjectManager->HandleToAnchorMap.Add(NewLuminArAnchor->Handle, OutAnchor);
+
+	return ELuminARAPIStatus::AR_SUCCESS;
 }
 
 void FLuminARSession::DetachAnchor(UARPin* Anchor)
diff --git a/Engine/Plugins/Lumin/MagicLeap/Source/MagicLeapAudio/Private/FakeAudioMixerMagicLeap.cpp b/Engine/Plugins/Lumin/MagicLeap/Source/MagicLeapAudio/Private/FakeAudioMixerMagicLeap.cpp
--- a/Engine/Plugins/Lumin/MagicLeap/Source/MagicLeapAudio/Private/FakeAudioMixerMagicLeap.cpp
+++ b/Engine/Plugins/Lumin/MagicLeap/Source/MagicLeapAudio/Private/FakeAudioMixerMagicLeap.cpp
@@ -127,10 +127,7 @@ namespace Audio
 			return false;
 		}
 
-		if (AudioStreamInfo.StreamState == EAudioOutputStreamState::Running)
-		{
-			StopGeneratingAudio();
-		}
+		StopGeneratingAudio();
 
 		check(AudioStreamInfo.StreamState == EAudioOutputStreamState::Stopped);
 		return true;
