This is a function stub from ML. Removing the implementation until it's fleshed out.

diff --git a/Engine/Plugins/Lumin/MagicLeap/Source/MagicLeapAR/Private/LuminARAPI.cpp b/Engine/Plugins/Lumin/MagicLeap/Source/MagicLeapAR/Private/LuminARAPI.cpp
--- a/Engine/Plugins/Lumin/MagicLeap/Source/MagicLeapAR/Private/LuminARAPI.cpp
+++ b/Engine/Plugins/Lumin/MagicLeap/Source/MagicLeapAR/Private/LuminARAPI.cpp
@@ -798,6 +798,7 @@ void ULuminARUObjectManager::DumpTrackableHandleMap(const FGuid& SessionHandle)
 
 ELuminARAPIStatus FLuminARSession::AcquireCameraImage(ULuminARCameraImage *&OutCameraImage)
 {
+	/*
 	ELuminARAPIStatus ApiStatus = ELuminARAPIStatus::AR_SUCCESS;
 #if WITH_MLSDK
 	if (LatestFrame == nullptr)
@@ -819,6 +820,8 @@ ELuminARAPIStatus FLuminARSession::AcquireCameraImage(ULuminARCameraImage *&OutC
 #endif // WITH_MLSDK
 
 	return ApiStatus;
+	*/
+	return ELuminARAPIStatus::AR_SUCCESS;
 }
 
 void* FLuminARSession::GetLatestFrameRawPointer()
