Fix compilation error introduced in CL 9960925

diff --git a/Engine/Plugins/Lumin/MagicLeap/Source/MagicLeapController/Private/MagicLeapController.cpp b/Engine/Plugins/Lumin/MagicLeap/Source/MagicLeapController/Private/MagicLeapController.cpp
--- a/Engine/Plugins/Lumin/MagicLeap/Source/MagicLeapController/Private/MagicLeapController.cpp
+++ b/Engine/Plugins/Lumin/MagicLeap/Source/MagicLeapController/Private/MagicLeapController.cpp
@@ -1343,11 +1343,11 @@ void FMagicLeapController::SendControllerEventsForHand(EControllerHand Hand)
 			MessageHandler->OnControllerAnalog(MLTouchToUnrealTrackpadAxis(Hand, 2),
 				DeviceIndex, CurrControllerState->TouchPosAndForce[0].Z);
 
-			MessageHandler->OnControllerAnalog(MLTouchToUnrealTrackpadAxis(Hand, 0, true),
+			MessageHandler->OnControllerAnalog(MLTouchToUnrealTrackpadAxisLegacy(Hand, 0),
 				DeviceIndex, CurrControllerState->TouchPosAndForce[0].X);
-			MessageHandler->OnControllerAnalog(MLTouchToUnrealTrackpadAxis(Hand, 1, true),
+			MessageHandler->OnControllerAnalog(MLTouchToUnrealTrackpadAxisLegacy(Hand, 1),
 				DeviceIndex, CurrControllerState->TouchPosAndForce[0].Y);
-			MessageHandler->OnControllerAnalog(MLTouchToUnrealTrackpadAxis(Hand, 2, true),
+			MessageHandler->OnControllerAnalog(MLTouchToUnrealTrackpadAxisLegacy(Hand, 2),
 				DeviceIndex, CurrControllerState->TouchPosAndForce[0].Z);
 		}
 		else
@@ -1360,9 +1360,9 @@ void FMagicLeapController::SendControllerEventsForHand(EControllerHand Hand)
 			MessageHandler->OnControllerAnalog(MLTouchToUnrealTrackpadAxis(Hand, 1), DeviceIndex, 0.0f);
 			MessageHandler->OnControllerAnalog(MLTouchToUnrealTrackpadAxis(Hand, 2), DeviceIndex, 0.0f);
 
-			MessageHandler->OnControllerAnalog(MLTouchToUnrealTrackpadAxis(Hand, 0, true), DeviceIndex, 0.0f);
-			MessageHandler->OnControllerAnalog(MLTouchToUnrealTrackpadAxis(Hand, 1, true), DeviceIndex, 0.0f);
-			MessageHandler->OnControllerAnalog(MLTouchToUnrealTrackpadAxis(Hand, 2, true), DeviceIndex, 0.0f);
+			MessageHandler->OnControllerAnalog(MLTouchToUnrealTrackpadAxisLegacy(Hand, 0), DeviceIndex, 0.0f);
+			MessageHandler->OnControllerAnalog(MLTouchToUnrealTrackpadAxisLegacy(Hand, 1), DeviceIndex, 0.0f);
+			MessageHandler->OnControllerAnalog(MLTouchToUnrealTrackpadAxisLegacy(Hand, 2), DeviceIndex, 0.0f);
 		}
 
 		// Analog touch coords
@@ -1401,7 +1401,7 @@ void FMagicLeapController::SendControllerEventsForHand(EControllerHand Hand)
 
 			MessageHandler->OnControllerAnalog(MLTriggerToUnrealTriggerAxis(Hand),
 				DeviceIndex, CurrControllerState->TriggerAnalog);
-			MessageHandler->OnControllerAnalogLegacy(MLTriggerToUnrealTriggerAxis(Hand),
+			MessageHandler->OnControllerAnalog(MLTriggerToUnrealTriggerAxisLegacy(Hand),
 				DeviceIndex, CurrControllerState->TriggerAnalog);
 		}
 	}
