We can't return null here. This isn't ideal as it causes us to use a fake audio device when the ML plugin is active, ZI and ML audio mixer are active, but another HMD plugin is actually being used. End result is no audio. Given the list of opt-ins to get to the crash, I'm fixing the crash for now. The way ML is hooking the audio mixer needs to change.

diff --git a/Engine/Plugins/Lumin/MagicLeap/Source/MagicLeapAudio/Private/MagicLeapAudioModule.cpp b/Engine/Plugins/Lumin/MagicLeap/Source/MagicLeapAudio/Private/MagicLeapAudioModule.cpp
--- a/Engine/Plugins/Lumin/MagicLeap/Source/MagicLeapAudio/Private/MagicLeapAudioModule.cpp
+++ b/Engine/Plugins/Lumin/MagicLeap/Source/MagicLeapAudio/Private/MagicLeapAudioModule.cpp
@@ -25,10 +25,9 @@ FAudioDevice* FMagicLeapAudioModule::CreateAudioDevice()
 		{
 			return new Audio::FMixerDevice(new Audio::FMixerPlatformMagicLeap());
 		}
-		return new Audio::FMixerDevice(new Audio::FFakeAudioMixerMagicLeap());;
 	}
 
-	return nullptr;
+	return new Audio::FMixerDevice(new Audio::FFakeAudioMixerMagicLeap());
 #elif PLATFORM_LUMIN
 	return new Audio::FMixerDevice(new Audio::FMixerPlatformMagicLeap());
 #else
