Trackpad Y axis input was being inadvertently sent to trackpad X

diff --git a/Engine/Plugins/Lumin/MagicLeap/Source/MagicLeapController/Private/MagicLeapControllerMappings.h b/Engine/Plugins/Lumin/MagicLeap/Source/MagicLeapController/Private/MagicLeapControllerMappings.h
--- a/Engine/Plugins/Lumin/MagicLeap/Source/MagicLeapController/Private/MagicLeapControllerMappings.h
+++ b/Engine/Plugins/Lumin/MagicLeap/Source/MagicLeapController/Private/MagicLeapControllerMappings.h
@@ -124,9 +124,9 @@ const FName& MLTouchToUnrealTrackpadAxis(EControllerHand Hand, uint32 TouchIndex
 	case 1:
 		if (Hand == EControllerHand::Left)
 		{
-			return FMagicLeapControllerKeyNames::Left_Trackpad_X_Name;;
+			return FMagicLeapControllerKeyNames::Left_Trackpad_Y_Name;
 		}
-		return FMagicLeapControllerKeyNames::Right_Trackpad_X_Name;
+		return FMagicLeapControllerKeyNames::Right_Trackpad_Y_Name;
 	case 2:
 		if (Hand == EControllerHand::Left)
 		{
