Fix copyright notices

diff --git a/Engine/Plugins/Lumin/MagicLeap/Source/MagicLeapHandMeshing/MagicLeapHandMeshing.Build.cs b/Engine/Plugins/Lumin/MagicLeap/Source/MagicLeapHandMeshing/MagicLeapHandMeshing.Build.cs
--- a/Engine/Plugins/Lumin/MagicLeap/Source/MagicLeapHandMeshing/MagicLeapHandMeshing.Build.cs
+++ b/Engine/Plugins/Lumin/MagicLeap/Source/MagicLeapHandMeshing/MagicLeapHandMeshing.Build.cs
@@ -1,4 +1,4 @@
-// Copyright 1998-2018 Epic Games, Inc. All Rights Reserved.
+// Copyright 1998-2019 Epic Games, Inc. All Rights Reserved.
 
 using System.IO;
 
diff --git a/Engine/Plugins/Lumin/MagicLeap/Source/MagicLeapHandMeshing/Private/MagicLeapHandMeshingComponent.cpp b/Engine/Plugins/Lumin/MagicLeap/Source/MagicLeapHandMeshing/Private/MagicLeapHandMeshingComponent.cpp
--- a/Engine/Plugins/Lumin/MagicLeap/Source/MagicLeapHandMeshing/Private/MagicLeapHandMeshingComponent.cpp
+++ b/Engine/Plugins/Lumin/MagicLeap/Source/MagicLeapHandMeshing/Private/MagicLeapHandMeshingComponent.cpp
@@ -1,4 +1,4 @@
-// Copyright 1998-2018 Epic Games, Inc. All Rights Reserved.
+// Copyright 1998-2019 Epic Games, Inc. All Rights Reserved.
 
 #include "MagicLeapHandMeshingComponent.h"
 #include "MagicLeapHandMeshingModule.h"
diff --git a/Engine/Plugins/Lumin/MagicLeap/Source/MagicLeapHandMeshing/Private/MagicLeapHandMeshingFunctionLibrary.cpp b/Engine/Plugins/Lumin/MagicLeap/Source/MagicLeapHandMeshing/Private/MagicLeapHandMeshingFunctionLibrary.cpp
--- a/Engine/Plugins/Lumin/MagicLeap/Source/MagicLeapHandMeshing/Private/MagicLeapHandMeshingFunctionLibrary.cpp
+++ b/Engine/Plugins/Lumin/MagicLeap/Source/MagicLeapHandMeshing/Private/MagicLeapHandMeshingFunctionLibrary.cpp
@@ -1,4 +1,4 @@
-// Copyright 1998-2018 Epic Games, Inc. All Rights Reserved.
+// Copyright 1998-2019 Epic Games, Inc. All Rights Reserved.
 
 #include "MagicLeapHandMeshingFunctionLibrary.h"
 #include "MagicLeapHandMeshingModule.h"
diff --git a/Engine/Plugins/Lumin/MagicLeap/Source/MagicLeapHandMeshing/Private/MagicLeapHandMeshingModule.cpp b/Engine/Plugins/Lumin/MagicLeap/Source/MagicLeapHandMeshing/Private/MagicLeapHandMeshingModule.cpp
--- a/Engine/Plugins/Lumin/MagicLeap/Source/MagicLeapHandMeshing/Private/MagicLeapHandMeshingModule.cpp
+++ b/Engine/Plugins/Lumin/MagicLeap/Source/MagicLeapHandMeshing/Private/MagicLeapHandMeshingModule.cpp
@@ -1,4 +1,4 @@
-// Copyright 1998-2018 Epic Games, Inc. All Rights Reserved.
+// Copyright 1998-2019 Epic Games, Inc. All Rights Reserved.
 
 #include "MagicLeapHandMeshingModule.h"
 #include "GameFramework/Actor.h"
diff --git a/Engine/Plugins/Lumin/MagicLeap/Source/MagicLeapHandMeshing/Private/MagicLeapHandMeshingModule.h b/Engine/Plugins/Lumin/MagicLeap/Source/MagicLeapHandMeshing/Private/MagicLeapHandMeshingModule.h
--- a/Engine/Plugins/Lumin/MagicLeap/Source/MagicLeapHandMeshing/Private/MagicLeapHandMeshingModule.h
+++ b/Engine/Plugins/Lumin/MagicLeap/Source/MagicLeapHandMeshing/Private/MagicLeapHandMeshingModule.h
@@ -1,4 +1,4 @@
-// Copyright 1998-2018 Epic Games, Inc. All Rights Reserved.
+// Copyright 1998-2019 Epic Games, Inc. All Rights Reserved.
 
 #pragma once
 
diff --git a/Engine/Plugins/Lumin/MagicLeap/Source/MagicLeapHandMeshing/Private/MagicLeapMeshTracker.cpp b/Engine/Plugins/Lumin/MagicLeap/Source/MagicLeapHandMeshing/Private/MagicLeapMeshTracker.cpp
--- a/Engine/Plugins/Lumin/MagicLeap/Source/MagicLeapHandMeshing/Private/MagicLeapMeshTracker.cpp
+++ b/Engine/Plugins/Lumin/MagicLeap/Source/MagicLeapHandMeshing/Private/MagicLeapMeshTracker.cpp
@@ -1,4 +1,5 @@
-// Copyright 1998-2018 Epic Games, Inc. All Rights Reserved.
+// Copyright 1998-2019 Epic Games, Inc. All Rights Reserved.
+
 #include "MagicLeapMeshTracker.h"
 #include "HeadMountedDisplayFunctionLibrary.h"
 #include "MagicLeapHandMeshingTypes.h"
diff --git a/Engine/Plugins/Lumin/MagicLeap/Source/MagicLeapHandMeshing/Private/MagicLeapMeshTracker.h b/Engine/Plugins/Lumin/MagicLeap/Source/MagicLeapHandMeshing/Private/MagicLeapMeshTracker.h
--- a/Engine/Plugins/Lumin/MagicLeap/Source/MagicLeapHandMeshing/Private/MagicLeapMeshTracker.h
+++ b/Engine/Plugins/Lumin/MagicLeap/Source/MagicLeapHandMeshing/Private/MagicLeapMeshTracker.h
@@ -1,4 +1,4 @@
-// Copyright 1998-2018 Epic Games, Inc. All Rights Reserved.
+// Copyright 1998-2019 Epic Games, Inc. All Rights Reserved.
 
 #pragma once
 
diff --git a/Engine/Plugins/Lumin/MagicLeap/Source/MagicLeapHandMeshing/Public/IMagicLeapHandMeshingModule.h b/Engine/Plugins/Lumin/MagicLeap/Source/MagicLeapHandMeshing/Public/IMagicLeapHandMeshingModule.h
--- a/Engine/Plugins/Lumin/MagicLeap/Source/MagicLeapHandMeshing/Public/IMagicLeapHandMeshingModule.h
+++ b/Engine/Plugins/Lumin/MagicLeap/Source/MagicLeapHandMeshing/Public/IMagicLeapHandMeshingModule.h
@@ -1,4 +1,4 @@
-// Copyright 1998-2018 Epic Games, Inc. All Rights Reserved.
+// Copyright 1998-2019 Epic Games, Inc. All Rights Reserved.
 
 #pragma once
 
diff --git a/Engine/Plugins/Lumin/MagicLeap/Source/MagicLeapHandMeshing/Public/MagicLeapHandMeshingComponent.h b/Engine/Plugins/Lumin/MagicLeap/Source/MagicLeapHandMeshing/Public/MagicLeapHandMeshingComponent.h
--- a/Engine/Plugins/Lumin/MagicLeap/Source/MagicLeapHandMeshing/Public/MagicLeapHandMeshingComponent.h
+++ b/Engine/Plugins/Lumin/MagicLeap/Source/MagicLeapHandMeshing/Public/MagicLeapHandMeshingComponent.h
@@ -1,4 +1,4 @@
-// Copyright 1998-2018 Epic Games, Inc. All Rights Reserved.
+// Copyright 1998-2019 Epic Games, Inc. All Rights Reserved.
 
 #pragma once
 
diff --git a/Engine/Plugins/Lumin/MagicLeap/Source/MagicLeapHandMeshing/Public/MagicLeapHandMeshingFunctionLibrary.h b/Engine/Plugins/Lumin/MagicLeap/Source/MagicLeapHandMeshing/Public/MagicLeapHandMeshingFunctionLibrary.h
--- a/Engine/Plugins/Lumin/MagicLeap/Source/MagicLeapHandMeshing/Public/MagicLeapHandMeshingFunctionLibrary.h
+++ b/Engine/Plugins/Lumin/MagicLeap/Source/MagicLeapHandMeshing/Public/MagicLeapHandMeshingFunctionLibrary.h
@@ -1,4 +1,5 @@
-// Copyright 1998-2018 Epic Games, Inc. All Rights Reserved.
+// Copyright 1998-2019 Epic Games, Inc. All Rights Reserved.
+
 #pragma once
 
 #include "Engine/Engine.h"
diff --git a/Engine/Plugins/Lumin/MagicLeap/Source/MagicLeapHandMeshing/Public/MagicLeapHandMeshingTypes.h b/Engine/Plugins/Lumin/MagicLeap/Source/MagicLeapHandMeshing/Public/MagicLeapHandMeshingTypes.h
--- a/Engine/Plugins/Lumin/MagicLeap/Source/MagicLeapHandMeshing/Public/MagicLeapHandMeshingTypes.h
+++ b/Engine/Plugins/Lumin/MagicLeap/Source/MagicLeapHandMeshing/Public/MagicLeapHandMeshingTypes.h
@@ -1,4 +1,5 @@
-// Copyright 1998-2018 Epic Games, Inc. All Rights Reserved.
+// Copyright 1998-2019 Epic Games, Inc. All Rights Reserved.
+
 #pragma once
 
 #include "Engine/Engine.h"
diff --git a/Engine/Plugins/Lumin/MagicLeapNetworking/Source/MagicLeapNetworking.Build.cs b/Engine/Plugins/Lumin/MagicLeapNetworking/Source/MagicLeapNetworking.Build.cs
--- a/Engine/Plugins/Lumin/MagicLeapNetworking/Source/MagicLeapNetworking.Build.cs
+++ b/Engine/Plugins/Lumin/MagicLeapNetworking/Source/MagicLeapNetworking.Build.cs
@@ -1,4 +1,4 @@
-// Copyright 1998-2018 Epic Games, Inc. All Rights Reserved.
+// Copyright 1998-2019 Epic Games, Inc. All Rights Reserved.
 
 using System.IO;
 
diff --git a/Engine/Plugins/Lumin/MagicLeapNetworking/Source/Private/MagicLeapNetworkingComponent.cpp b/Engine/Plugins/Lumin/MagicLeapNetworking/Source/Private/MagicLeapNetworkingComponent.cpp
--- a/Engine/Plugins/Lumin/MagicLeapNetworking/Source/Private/MagicLeapNetworkingComponent.cpp
+++ b/Engine/Plugins/Lumin/MagicLeapNetworking/Source/Private/MagicLeapNetworkingComponent.cpp
@@ -1,4 +1,4 @@
-// Copyright 1998-2018 Epic Games, Inc. All Rights Reserved.
+// Copyright 1998-2019 Epic Games, Inc. All Rights Reserved.
 
 #include "MagicLeapNetworkingComponent.h"
 #include "MagicLeapNetworkingPlugin.h"
diff --git a/Engine/Plugins/Lumin/MagicLeapNetworking/Source/Private/MagicLeapNetworkingFunctionLibrary.cpp b/Engine/Plugins/Lumin/MagicLeapNetworking/Source/Private/MagicLeapNetworkingFunctionLibrary.cpp
--- a/Engine/Plugins/Lumin/MagicLeapNetworking/Source/Private/MagicLeapNetworkingFunctionLibrary.cpp
+++ b/Engine/Plugins/Lumin/MagicLeapNetworking/Source/Private/MagicLeapNetworkingFunctionLibrary.cpp
@@ -1,4 +1,4 @@
-// Copyright 1998-2018 Epic Games, Inc. All Rights Reserved.
+// Copyright 1998-2019 Epic Games, Inc. All Rights Reserved.
 
 #include "MagicLeapNetworkingFunctionLibrary.h"
 #include "MagicLeapNetworkingPlugin.h"
diff --git a/Engine/Plugins/Lumin/MagicLeapNetworking/Source/Private/MagicLeapNetworkingPlugin.cpp b/Engine/Plugins/Lumin/MagicLeapNetworking/Source/Private/MagicLeapNetworkingPlugin.cpp
--- a/Engine/Plugins/Lumin/MagicLeapNetworking/Source/Private/MagicLeapNetworkingPlugin.cpp
+++ b/Engine/Plugins/Lumin/MagicLeapNetworking/Source/Private/MagicLeapNetworkingPlugin.cpp
@@ -1,4 +1,4 @@
-// Copyright 1998-2018 Epic Games, Inc. All Rights Reserved.
+// Copyright 1998-2019 Epic Games, Inc. All Rights Reserved.
 
 #include "MagicLeapNetworkingPlugin.h"
 #include "Async/Async.h"
diff --git a/Engine/Plugins/Lumin/MagicLeapNetworking/Source/Private/MagicLeapNetworkingPlugin.h b/Engine/Plugins/Lumin/MagicLeapNetworking/Source/Private/MagicLeapNetworkingPlugin.h
--- a/Engine/Plugins/Lumin/MagicLeapNetworking/Source/Private/MagicLeapNetworkingPlugin.h
+++ b/Engine/Plugins/Lumin/MagicLeapNetworking/Source/Private/MagicLeapNetworkingPlugin.h
@@ -1,4 +1,4 @@
-// Copyright 1998-2018 Epic Games, Inc. All Rights Reserved.
+// Copyright 1998-2019 Epic Games, Inc. All Rights Reserved.
 
 #pragma once
 
diff --git a/Engine/Plugins/Lumin/MagicLeapNetworking/Source/Public/IMagicLeapNetworkingPlugin.h b/Engine/Plugins/Lumin/MagicLeapNetworking/Source/Public/IMagicLeapNetworkingPlugin.h
--- a/Engine/Plugins/Lumin/MagicLeapNetworking/Source/Public/IMagicLeapNetworkingPlugin.h
+++ b/Engine/Plugins/Lumin/MagicLeapNetworking/Source/Public/IMagicLeapNetworkingPlugin.h
@@ -1,4 +1,4 @@
-// Copyright 1998-2018 Epic Games, Inc. All Rights Reserved.
+// Copyright 1998-2019 Epic Games, Inc. All Rights Reserved.
 
 #pragma once
 
diff --git a/Engine/Plugins/Lumin/MagicLeapNetworking/Source/Public/MagicLeapNetworkingComponent.h b/Engine/Plugins/Lumin/MagicLeapNetworking/Source/Public/MagicLeapNetworkingComponent.h
--- a/Engine/Plugins/Lumin/MagicLeapNetworking/Source/Public/MagicLeapNetworkingComponent.h
+++ b/Engine/Plugins/Lumin/MagicLeapNetworking/Source/Public/MagicLeapNetworkingComponent.h
@@ -1,4 +1,4 @@
-// Copyright 1998-2018 Epic Games, Inc. All Rights Reserved.
+// Copyright 1998-2019 Epic Games, Inc. All Rights Reserved.
 
 #pragma once
 
diff --git a/Engine/Plugins/Lumin/MagicLeapNetworking/Source/Public/MagicLeapNetworkingFunctionLibrary.h b/Engine/Plugins/Lumin/MagicLeapNetworking/Source/Public/MagicLeapNetworkingFunctionLibrary.h
--- a/Engine/Plugins/Lumin/MagicLeapNetworking/Source/Public/MagicLeapNetworkingFunctionLibrary.h
+++ b/Engine/Plugins/Lumin/MagicLeapNetworking/Source/Public/MagicLeapNetworkingFunctionLibrary.h
@@ -1,4 +1,5 @@
-// Copyright 1998-2018 Epic Games, Inc. All Rights Reserved.
+// Copyright 1998-2019 Epic Games, Inc. All Rights Reserved.
+
 #pragma once
 
 #include "Engine/Engine.h"
diff --git a/Engine/Plugins/Lumin/MagicLeapNetworking/Source/Public/MagicLeapNetworkingTypes.h b/Engine/Plugins/Lumin/MagicLeapNetworking/Source/Public/MagicLeapNetworkingTypes.h
--- a/Engine/Plugins/Lumin/MagicLeapNetworking/Source/Public/MagicLeapNetworkingTypes.h
+++ b/Engine/Plugins/Lumin/MagicLeapNetworking/Source/Public/MagicLeapNetworkingTypes.h
@@ -1,4 +1,5 @@
-// Copyright 1998-2018 Epic Games, Inc. All Rights Reserved.
+// Copyright 1998-2019 Epic Games, Inc. All Rights Reserved.
+
 #pragma once
 
 #include "Engine/Engine.h"
diff --git a/Engine/Source/Runtime/Core/Private/Lumin/LuminPlatformDelegates.cpp b/Engine/Source/Runtime/Core/Private/Lumin/LuminPlatformDelegates.cpp
--- a/Engine/Source/Runtime/Core/Private/Lumin/LuminPlatformDelegates.cpp
+++ b/Engine/Source/Runtime/Core/Private/Lumin/LuminPlatformDelegates.cpp
@@ -1,3 +1,5 @@
+// Copyright 1998-2019 Epic Games, Inc. All Rights Reserved.
+
 #include "Lumin/LuminPlatformDelegates.h"
 
 FLuminDelegates::FLuminAppStartupArgumentsDelegate FLuminDelegates::LuminAppReceivedStartupArgumentsDelegate;
diff --git a/Engine/Source/Runtime/Core/Public/Lumin/LuminPlatformDelegates.h b/Engine/Source/Runtime/Core/Public/Lumin/LuminPlatformDelegates.h
--- a/Engine/Source/Runtime/Core/Public/Lumin/LuminPlatformDelegates.h
+++ b/Engine/Source/Runtime/Core/Public/Lumin/LuminPlatformDelegates.h
@@ -1,4 +1,4 @@
-// Copyright 1998-2018 Epic Games, Inc. All Rights Reserved.
+// Copyright 1998-2019 Epic Games, Inc. All Rights Reserved.
 
 #pragma once
 
