Missed a line. Also removed GEngine usage since we have a reference to the stereo rendering device as a member

diff --git a/Engine/Plugins/Lumin/MagicLeap/Source/MagicLeap/Private/MagicLeapXRCamera.cpp b/Engine/Plugins/Lumin/MagicLeap/Source/MagicLeap/Private/MagicLeapXRCamera.cpp
--- a/Engine/Plugins/Lumin/MagicLeap/Source/MagicLeap/Private/MagicLeapXRCamera.cpp
+++ b/Engine/Plugins/Lumin/MagicLeap/Source/MagicLeap/Private/MagicLeapXRCamera.cpp
@@ -16,9 +16,10 @@ void FMagicLeapXRCamera::PreRenderView_RenderThread(FRHICommandListImmediate& RH
 {
 #if WITH_MLSDK
 	// this needs to happen before the FDefaultXRCamera call, because UpdateProjectionMatrix is somewhat destructive. 
-	if (IStereoRendering::IsStereoEyePass(View.StereoPass))
+	
+	if (MagicLeapSystem.IsStereoEyePass(View.StereoPass))
 	{
-		const int EyeIdx = IStereoRendering::IsAPrimaryView(View, GEngine->StereoRenderingDevice) ? 0 : 1;
+		const int EyeIdx = IStereoRendering::IsAPrimaryView(View, MagicLeapSystem.AsShared()) ? 0 : 1;
 
 		const FTrackingFrame& Frame = MagicLeapSystem.GetCurrentFrame();
 
