UE-64550 [CrashReport] <unknown module>!UE4Editor_BodyState!FCSPose<FCompactPose>::CalculateComponentSpaceTransform()
-Speculative fix (found by [REDACTED_PERSON]), the code was definately wrong here and could plausibly lead to NaN values.
-We want to move the local copy we have made into livelink, not the passed in data of unreliable lifetime.

diff --git a/Engine/Plugins/Lumin/MagicLeap/Source/MagicLeapHandTracking/Private/MagicLeapHandTrackingLiveLink.cpp b/Engine/Plugins/Lumin/MagicLeap/Source/MagicLeapHandTracking/Private/MagicLeapHandTrackingLiveLink.cpp
--- a/Engine/Plugins/Lumin/MagicLeap/Source/MagicLeapHandTracking/Private/MagicLeapHandTrackingLiveLink.cpp
+++ b/Engine/Plugins/Lumin/MagicLeap/Source/MagicLeapHandTracking/Private/MagicLeapHandTrackingLiveLink.cpp
@@ -202,7 +202,7 @@ void FMagicLeapHandTracking::UpdateLiveLink()
 		FLiveLinkFrameDataStruct NewLiveLinkRightFrame;
 		NewLiveLinkLeftFrame.InitializeWith(LiveLinkLeftFrame);
 		NewLiveLinkRightFrame.InitializeWith(LiveLinkRightFrame);
-		LiveLinkClient->PushSubjectFrameData_AnyThread(LeftKey, MoveTemp(LiveLinkLeftFrame));
+		LiveLinkClient->PushSubjectFrameData_AnyThread(LeftKey, MoveTemp(NewLiveLinkLeftFrame));
 		LiveLinkClient->PushSubjectFrameData_AnyThread(RightKey, MoveTemp(NewLiveLinkRightFrame));
 	}
 }
