Merge from MagicLeap-Staging: "There was an ask from Epic few months ago that all public facing data types in ML plugins should have the 'MagicLeap' prefix. This CL addresses the data types that were missed during the last update."

diff --git a/Engine/Plugins/Lumin/MagicLeap/Source/MagicLeapImageTracker/Public/MagicLeapImageTrackerComponent.h b/Engine/Plugins/Lumin/MagicLeap/Source/MagicLeapImageTracker/Public/MagicLeapImageTrackerComponent.h
--- a/Engine/Plugins/Lumin/MagicLeap/Source/MagicLeapImageTracker/Public/MagicLeapImageTrackerComponent.h
+++ b/Engine/Plugins/Lumin/MagicLeap/Source/MagicLeapImageTracker/Public/MagicLeapImageTrackerComponent.h
@@ -79,19 +79,19 @@ public:
 //private:
 	/** Activated when the target image is successfully set. */
 	UPROPERTY(BlueprintAssignable, Category = "ImageTracking | MagicLeap", meta = (AllowPrivateAccess = true))
-	FSetImageTargetSucceededMulti OnSetImageTargetSucceeded;
+	FMagicLeapSetImageTargetSucceededMulti OnSetImageTargetSucceeded;
 
 	/** Activated when the target image fails to be set. */
 	UPROPERTY(BlueprintAssignable, Category = "ImageTracking | MagicLeap", meta = (AllowPrivateAccess = true))
-	FSetImageTargetFailedMulti OnSetImageTargetFailed;
+	FMagicLeapSetImageTargetFailedMulti OnSetImageTargetFailed;
 
 	/** Activated when the target image becomes visible to the camera */
 	UPROPERTY(BlueprintAssignable, Category = "ImageTracking | MagicLeap", meta = (AllowPrivateAccess = true))
-	FImageTargetFoundMulti OnImageTargetFound;
+	FMagicLeapImageTargetFoundMulti OnImageTargetFound;
 
 	/** Activated the target image becomes invisible to the camera */
 	UPROPERTY(BlueprintAssignable, Category = "ImageTracking | MagicLeap", meta = (AllowPrivateAccess = true))
-	FImageTargetLostMulti OnImageTargetLost;
+	FMagicLeapImageTargetLostMulti OnImageTargetLost;
 
 	/**
 	  Activated when the target image is tracked with low confidence.
@@ -103,7 +103,7 @@ public:
 	  status will change to NotTracked.
 	*/
 	UPROPERTY(BlueprintAssignable, Category = "ImageTracking | MagicLeap", meta = (AllowPrivateAccess = true))
-	FImageTargetUnreliableTrackingMulti OnImageTargetUnreliableTracking;
+	FMagicLeapImageTargetUnreliableTrackingMulti OnImageTargetUnreliableTracking;
 
 private:
 	bool bIsTracking;
diff --git a/Engine/Plugins/Lumin/MagicLeap/Source/MagicLeapImageTracker/Public/MagicLeapImageTrackerTypes.h b/Engine/Plugins/Lumin/MagicLeap/Source/MagicLeapImageTracker/Public/MagicLeapImageTrackerTypes.h
--- a/Engine/Plugins/Lumin/MagicLeap/Source/MagicLeapImageTracker/Public/MagicLeapImageTrackerTypes.h
+++ b/Engine/Plugins/Lumin/MagicLeap/Source/MagicLeapImageTracker/Public/MagicLeapImageTrackerTypes.h
@@ -13,35 +13,35 @@
 	@param NewUnreliableLocation The new location of the target image (which may or may not be accurate).
 	@param NewUnreliableRotation The new rotation of the target image (which may or may not be accurate).
 */
-DECLARE_DYNAMIC_DELEGATE_FourParams(FImageTargetUnreliableTracking, const FVector&, LastTrackedLocation, const FRotator&, LastTrackedRotation, const FVector&, NewUnreliableLocation, const FRotator&, NewUnreliableRotation);
-DECLARE_DYNAMIC_MULTICAST_DELEGATE_FourParams(FImageTargetUnreliableTrackingMulti, const FVector&, LastTrackedLocation, const FRotator&, LastTrackedRotation, const FVector&, NewUnreliableLocation, const FRotator&, NewUnreliableRotation);
+DECLARE_DYNAMIC_DELEGATE_FourParams(FMagicLeapImageTargetUnreliableTracking, const FVector&, LastTrackedLocation, const FRotator&, LastTrackedRotation, const FVector&, NewUnreliableLocation, const FRotator&, NewUnreliableRotation);
+DECLARE_DYNAMIC_MULTICAST_DELEGATE_FourParams(FMagicLeapImageTargetUnreliableTrackingMulti, const FVector&, LastTrackedLocation, const FRotator&, LastTrackedRotation, const FVector&, NewUnreliableLocation, const FRotator&, NewUnreliableRotation);
 
 /**
 	Delegate used to notify the instigating blueprint that the target image's location/rotation has changed.
 	@param NewLocation The new location of the target image (which may or may not be accurate).
 	@param NewRotation The new rotation of the target image (which may or may not be accurate).
 */
-DECLARE_DYNAMIC_DELEGATE_TwoParams(FImageTargetReliableTracking, const FVector&, NewLocation, const FRotator&, NewRotation);
-DECLARE_DYNAMIC_MULTICAST_DELEGATE_TwoParams(FImageTargetReliableTrackingMulti, const FVector&, NewLocation, const FRotator&, NewRotation);
+DECLARE_DYNAMIC_DELEGATE_TwoParams(FMagicLeapImageTargetReliableTracking, const FVector&, NewLocation, const FRotator&, NewRotation);
+DECLARE_DYNAMIC_MULTICAST_DELEGATE_TwoParams(FMagicLeapImageTargetReliableTrackingMulti, const FVector&, NewLocation, const FRotator&, NewRotation);
 
 /** Delegate used to notify LuminARImageTracker that the target image was successfully set. */
-DECLARE_DELEGATE_OneParam(FSetImageTargetSucceededStaticDelegate, FMagicLeapImageTrackerTarget& /*Target*/);
+DECLARE_DELEGATE_OneParam(FMagicLeapSetImageTargetSucceededStaticDelegate, FMagicLeapImageTrackerTarget& /*Target*/);
 
 /** Delegate used to notify the instigating blueprint that the target image was successfully set. */
-DECLARE_DYNAMIC_DELEGATE(FSetImageTargetSucceeded);
-DECLARE_DYNAMIC_MULTICAST_DELEGATE(FSetImageTargetSucceededMulti);
+DECLARE_DYNAMIC_DELEGATE(FMagicLeapSetImageTargetSucceeded);
+DECLARE_DYNAMIC_MULTICAST_DELEGATE(FMagicLeapSetImageTargetSucceededMulti);
 
 /** Delegate used to notify the instigating blueprint that the target image failed to be set. */
-DECLARE_DYNAMIC_DELEGATE(FSetImageTargetFailed);
-DECLARE_DYNAMIC_MULTICAST_DELEGATE(FSetImageTargetFailedMulti);
+DECLARE_DYNAMIC_DELEGATE(FMagicLeapSetImageTargetFailed);
+DECLARE_DYNAMIC_MULTICAST_DELEGATE(FMagicLeapSetImageTargetFailedMulti);
 
 /** Delegate used to notify the instigating blueprint that the target image is currently visible to the camera */
-DECLARE_DYNAMIC_DELEGATE(FImageTargetFound);
-DECLARE_DYNAMIC_MULTICAST_DELEGATE(FImageTargetFoundMulti);
+DECLARE_DYNAMIC_DELEGATE(FMagicLeapImageTargetFound);
+DECLARE_DYNAMIC_MULTICAST_DELEGATE(FMagicLeapImageTargetFoundMulti);
 
 /** Delegate used to notify the instigating blueprint that the target image just became invisible to the camera */
-DECLARE_DYNAMIC_DELEGATE(FImageTargetLost);
-DECLARE_DYNAMIC_MULTICAST_DELEGATE(FImageTargetLostMulti);
+DECLARE_DYNAMIC_DELEGATE(FMagicLeapImageTargetLost);
+DECLARE_DYNAMIC_MULTICAST_DELEGATE(FMagicLeapImageTargetLostMulti);
 
 USTRUCT()
 struct FMagicLeapImageTrackerTarget
@@ -61,12 +61,12 @@ struct FMagicLeapImageTrackerTarget
 	FVector UnreliableLocation;
 	FRotator UnreliableRotation;
 	bool bUseUnreliablePose;
-	FSetImageTargetSucceededMulti OnSetImageTargetSucceeded;
-	FSetImageTargetFailedMulti OnSetImageTargetFailed;
-	FImageTargetFoundMulti OnImageTargetFound;
-	FImageTargetLostMulti OnImageTargetLost;
-	FImageTargetUnreliableTrackingMulti OnImageTargetUnreliableTracking;
-	FSetImageTargetSucceededStaticDelegate SetImageTargetSucceededDelegate;
+	FMagicLeapSetImageTargetSucceededMulti OnSetImageTargetSucceeded;
+	FMagicLeapSetImageTargetFailedMulti OnSetImageTargetFailed;
+	FMagicLeapImageTargetFoundMulti OnImageTargetFound;
+	FMagicLeapImageTargetLostMulti OnImageTargetLost;
+	FMagicLeapImageTargetUnreliableTrackingMulti OnImageTargetUnreliableTracking;
+	FMagicLeapSetImageTargetSucceededStaticDelegate SetImageTargetSucceededDelegate;
 
 	FMagicLeapImageTrackerTarget()
 	: Name(TEXT("Undefined"))
