Merge from MagicLeap Staging: Fixes an vulkan validation layer error for missed layout transition. Can be tested on Windows as well by using running VRPreview with Zero Iteration.

diff --git a/Engine/Plugins/Lumin/MagicLeap/Source/MagicLeapHelperVulkan/Private/MagicLeapHelperVulkan.cpp b/Engine/Plugins/Lumin/MagicLeap/Source/MagicLeapHelperVulkan/Private/MagicLeapHelperVulkan.cpp
--- a/Engine/Plugins/Lumin/MagicLeap/Source/MagicLeapHelperVulkan/Private/MagicLeapHelperVulkan.cpp
+++ b/Engine/Plugins/Lumin/MagicLeap/Source/MagicLeapHelperVulkan/Private/MagicLeapHelperVulkan.cpp
@@ -168,6 +168,20 @@ uint64 FMagicLeapHelperVulkan::AliasImageSRGB(const uint64 Allocation, const uin
 
 	VERIFYVULKANRESULT(VulkanRHI::vkBindImageMemory(Device->GetInstanceHandle(), Result, (VkDeviceMemory)Allocation, AllocationOffset));
 
+	check(Result != VK_NULL_HANDLE);
+
+	FVulkanCommandBufferManager* CmdBufferMgr = RHI->GetDevice()->GetImmediateContext().GetCommandBufferManager();
+	FVulkanCmdBuffer* CmdBuffer = CmdBufferMgr->GetUploadCmdBuffer();
+
+	VkImageSubresourceRange Range;
+	Range.aspectMask = VK_IMAGE_ASPECT_COLOR_BIT;
+	Range.baseMipLevel = 0;
+	Range.levelCount = 1;
+	Range.baseArrayLayer = 0;
+	Range.layerCount = 1;
+
+	RHI->VulkanSetImageLayout(CmdBuffer->GetHandle(), Result, VK_IMAGE_LAYOUT_UNDEFINED, VK_IMAGE_LAYOUT_TRANSFER_SRC_OPTIMAL, Range);
+
 	return (uint64)Result;
 #else
 	return 0;
