Merged from MagicLeap-Staging:

Various fixes for ML hand meshing.

Remove explicit caching and re-sending of the current mesh data. This is not required because MRMesh keeps a local cache of the current data to be rendered.  We only need to send new data when we want to update the mesh.
Fix flickering caused by garbage data being used for the bounds in a frustum/bounds visibility check. The data was garbage because it was never supplied.  Supply a box based on the hmd position and orientation.

Fix vertex normal calculation for ML hand meshing (using incorrect normals was crashing on low poly meshes). Generate normals by adding the normal of each triangle that a vertex is part of and then normalizing the result (this is done during tangent generation to remove the need to add a separate iteration over the mesh just to do normalization). If triangle weights should be used for this, multiply the normal by the area of the triangle first. Add blueprint callable to set whether to use triangle weights when calculating the vertex normals. False by default, as this is slightly more expensive and doesn't have a huge impact.
Note: the order of the cross product in FMagicLeapMeshTracker::GetResult has been intentionally swapped to account for the left handed coordinate system of unreal. One sided materials should work on the hands now (a previous change to the hand mesh winding only seemed to work because the normals weren't correct at that time).

diff --git a/Engine/Plugins/Lumin/MagicLeap/Source/MagicLeapHandMeshing/Private/MagicLeapHandMeshingComponent.cpp b/Engine/Plugins/Lumin/MagicLeap/Source/MagicLeapHandMeshing/Private/MagicLeapHandMeshingComponent.cpp
--- a/Engine/Plugins/Lumin/MagicLeap/Source/MagicLeapHandMeshing/Private/MagicLeapHandMeshingComponent.cpp
+++ b/Engine/Plugins/Lumin/MagicLeap/Source/MagicLeapHandMeshing/Private/MagicLeapHandMeshingComponent.cpp
@@ -24,3 +24,8 @@ bool UMagicLeapHandMeshingComponent::DisconnectMRMesh(class UMRMeshComponent* In
 {
 	return GetMagicLeapHandMeshingModule().DisconnectMRMesh(InMRMeshPtr);
 }
+
+void UMagicLeapHandMeshingComponent::SetUseWeightedNormals(const bool bInUseWeightedNormals)
+{
+	GetMagicLeapHandMeshingModule().SetUseWeightedNormals(bInUseWeightedNormals);
+}
diff --git a/Engine/Plugins/Lumin/MagicLeap/Source/MagicLeapHandMeshing/Private/MagicLeapHandMeshingModule.cpp b/Engine/Plugins/Lumin/MagicLeap/Source/MagicLeapHandMeshing/Private/MagicLeapHandMeshingModule.cpp
--- a/Engine/Plugins/Lumin/MagicLeap/Source/MagicLeapHandMeshing/Private/MagicLeapHandMeshingModule.cpp
+++ b/Engine/Plugins/Lumin/MagicLeap/Source/MagicLeapHandMeshing/Private/MagicLeapHandMeshingModule.cpp
@@ -45,6 +45,11 @@ bool FMagicLeapHandMeshingModule::DisconnectMRMesh(class UMRMeshComponent* InMRM
 	return MeshTracker.DisconnectMRMesh(InMRMeshPtr);
 }
 
+void FMagicLeapHandMeshingModule::SetUseWeightedNormals(const bool bInUseWeightedNormals)
+{
+	MeshTracker.SetUseWeightedNormals(bInUseWeightedNormals);
+}
+
 bool FMagicLeapHandMeshingModule::HasMRMesh() const
 {
 	return MeshTracker.HasMRMesh();
diff --git a/Engine/Plugins/Lumin/MagicLeap/Source/MagicLeapHandMeshing/Private/MagicLeapHandMeshingModule.h b/Engine/Plugins/Lumin/MagicLeap/Source/MagicLeapHandMeshing/Private/MagicLeapHandMeshingModule.h
--- a/Engine/Plugins/Lumin/MagicLeap/Source/MagicLeapHandMeshing/Private/MagicLeapHandMeshingModule.h
+++ b/Engine/Plugins/Lumin/MagicLeap/Source/MagicLeapHandMeshing/Private/MagicLeapHandMeshingModule.h
@@ -19,6 +19,7 @@ public:
 	bool DestroyClient();
 	bool ConnectMRMesh(UMRMeshComponent* InMRMeshPtr);
 	bool DisconnectMRMesh(class UMRMeshComponent* InMRMeshPtr);
+	void SetUseWeightedNormals(const bool bInUseWeightedNormals);
 	bool HasMRMesh() const;
 	bool HasClient() const;
 
diff --git a/Engine/Plugins/Lumin/MagicLeap/Source/MagicLeapHandMeshing/Private/MagicLeapMeshTracker.cpp b/Engine/Plugins/Lumin/MagicLeap/Source/MagicLeapHandMeshing/Private/MagicLeapMeshTracker.cpp
--- a/Engine/Plugins/Lumin/MagicLeap/Source/MagicLeapHandMeshing/Private/MagicLeapMeshTracker.cpp
+++ b/Engine/Plugins/Lumin/MagicLeap/Source/MagicLeapHandMeshing/Private/MagicLeapMeshTracker.cpp
@@ -13,9 +13,9 @@ FMagicLeapMeshTracker::FMagicLeapMeshTracker()
 , CurrentMeshRequest(ML_INVALID_HANDLE)
 #endif //WITH_MLSDK
 , bCreating(false)
+, bUseWeightedNormals(false)
 , MeshBrickIndex(0)
 , MRMesh(nullptr)
-, CurrentMeshDataCache(nullptr)
 {
 };
 
@@ -123,6 +123,11 @@ bool FMagicLeapMeshTracker::DisconnectMRMesh(class UMRMeshComponent* InMRMeshPtr
 	return true;
 }
 
+void FMagicLeapMeshTracker::SetUseWeightedNormals(const bool bInUseWeightedNormals)
+{
+	bUseWeightedNormals = bInUseWeightedNormals;
+}
+
 bool FMagicLeapMeshTracker::Update()
 {
 	if (bCreating)
@@ -203,38 +208,28 @@ bool FMagicLeapMeshTracker::GetMeshResult()
 				CurrentMeshRequest = ML_INVALID_HANDLE;
 				return true;
 			}
-
-			if (CurrentMeshDataCache.Get() != nullptr)
-			{
-				// Get/create brick ID for this mesh GUID
-				auto MeshBrickInfo = GetBrickInfo(CurrentMeshDataCache->BlockID, true);
-				// Create/update brick
-				static_cast<IMRMesh*>(MRMesh)->SendBrickData(
-					IMRMesh::FSendBrickDataArgs
-					{
-						TSharedPtr<IMRMesh::FBrickDataReceipt, ESPMode::ThreadSafe>
-						(new FMeshTrackerBrickDataReceipt(CurrentMeshDataCache)),
-						*MeshBrickInfo,
-						CurrentMeshDataCache->WorldVertices,
-						CurrentMeshDataCache->UV0,
-						CurrentMeshDataCache->Tangents,
-						CurrentMeshDataCache->VertexColors,
-						CurrentMeshDataCache->Triangles
-					}
-				);
-			}
 			// Mesh request pending...
 			return false;
 		}
 		else
 		{
+			// Create a bounding box based on the hmd position and rotation
+			FRotator HMDRotation;
+			FVector HMDPosition;
+			UHeadMountedDisplayFunctionLibrary::GetOrientationAndPosition(HMDRotation, HMDPosition);
+			// Put the center 1/2 meter in front of the face and set the radius to 1 meter
+			const FTransform TrackingToWorld = UHeadMountedDisplayFunctionLibrary::GetTrackingToWorldTransform(nullptr);
+			FVector Center = TrackingToWorld.TransformPosition(HMDPosition + ((WorldToMetersScale / 2.f) * HMDRotation.Vector()));
+			FVector BoxExtent(WorldToMetersScale);
+			FBox Bounds(Center - BoxExtent, Center + BoxExtent);
+
 			FVector VertexOffset = UHeadMountedDisplayFunctionLibrary::GetTrackingToWorldTransform(MRMesh).Inverse().GetLocation();
 			for (uint32_t MeshIndex = 0; MeshIndex < Mesh.data_count; ++MeshIndex)
 			{
 				const auto &MeshData = Mesh.data[MeshIndex];
 				auto BlockID = FGuid();// MeshData.id.data[0], MeshData.id.data[0] >> 32, MeshData.id.data[1], MeshData.id.data[1] >> 32);
 				// Acquire mesh data cache and mark its brick ID
-				CurrentMeshDataCache = AquireMeshDataCache();
+				FMLCachedMeshData::SharedPtr CurrentMeshDataCache = AquireMeshDataCache();
 				CurrentMeshDataCache->BlockID = BlockID;
 
 				// Pull vertices
@@ -257,24 +252,39 @@ bool FMagicLeapMeshTracker::GetMeshResult()
 					CurrentMeshDataCache->Triangles.Add(static_cast<uint32>(MeshData.index[i+1]));
 				}
 
-				// Pull normals
-				CurrentMeshDataCache->Normals.Reserve(MeshData.vertex_count);
-				for (int32 iNormal = 0; iNormal < CurrentMeshDataCache->Triangles.Num()-2; ++iNormal)
+				// Generate normals
+				CurrentMeshDataCache->Normals.SetNumZeroed(MeshData.vertex_count);
+				for (int32 iNormal = 0; iNormal < CurrentMeshDataCache->Triangles.Num()-2; iNormal += 3)
 				{
 					FVector A = CurrentMeshDataCache->WorldVertices[CurrentMeshDataCache->Triangles[iNormal]];
 					FVector B = CurrentMeshDataCache->WorldVertices[CurrentMeshDataCache->Triangles[iNormal+1]];
 					FVector C = CurrentMeshDataCache->WorldVertices[CurrentMeshDataCache->Triangles[iNormal+2]];
+					// Get the normal for this triangle.
 					FVector AToB = B - A;
-					AToB.Normalize();
 					FVector AToC = C - A;
-					AToC.Normalize();
-					FVector Normal = FVector::CrossProduct(AToB, AToC);
-					CurrentMeshDataCache->Normals.Add(Normal);
+					FVector Normal = FVector::CrossProduct(AToC, AToB);
+					// Weight it based on the area of the triangle, if requested.  Otherwise, just normalize it.
+					if (bUseWeightedNormals)
+					{
+						float TriangleSize = 0.5f * Normal.Size();
+						Normal.Normalize();
+						Normal *= TriangleSize;
+					}
+					else
+					{
+						Normal.Normalize();
+					}
+					// Add to the normals of each vertex of the triangle.  The final normals will be Normalized while iterating to get tangents, below.
+					CurrentMeshDataCache->Normals[CurrentMeshDataCache->Triangles[iNormal]] += Normal;
+					CurrentMeshDataCache->Normals[CurrentMeshDataCache->Triangles[iNormal+1]] += Normal;
+					CurrentMeshDataCache->Normals[CurrentMeshDataCache->Triangles[iNormal+2]] += Normal;
 				}
 				
 				CurrentMeshDataCache->Tangents.Reserve(MeshData.vertex_count * 2);
 				for (uint32_t t = 0; t < MeshData.vertex_count; ++t)
 				{
+					// Normals aren't normalized above due to the iterative nature of their generation.  Normalize here before getting their tangents.
+					CurrentMeshDataCache->Normals[t].Normalize();
 					const FVector& Norm = CurrentMeshDataCache->Normals[t];
 
 					// Calculate tangent
@@ -324,7 +334,8 @@ bool FMagicLeapMeshTracker::GetMeshResult()
 						CurrentMeshDataCache->UV0,
 						CurrentMeshDataCache->Tangents,
 						CurrentMeshDataCache->VertexColors,
-						CurrentMeshDataCache->Triangles
+						CurrentMeshDataCache->Triangles,
+						Bounds
 					}
 				);
 			}
@@ -332,7 +343,6 @@ bool FMagicLeapMeshTracker::GetMeshResult()
 			// All meshes pulled and/or updated; free the ML resource
 			MLHandMeshingFreeResource(MeshTracker, &CurrentMeshRequest);
 			CurrentMeshRequest = ML_INVALID_HANDLE;
-			CurrentMeshDataCache = nullptr;
 			return true;
 		}
 	}
diff --git a/Engine/Plugins/Lumin/MagicLeap/Source/MagicLeapHandMeshing/Private/MagicLeapMeshTracker.h b/Engine/Plugins/Lumin/MagicLeap/Source/MagicLeapHandMeshing/Private/MagicLeapMeshTracker.h
--- a/Engine/Plugins/Lumin/MagicLeap/Source/MagicLeapHandMeshing/Private/MagicLeapMeshTracker.h
+++ b/Engine/Plugins/Lumin/MagicLeap/Source/MagicLeapHandMeshing/Private/MagicLeapMeshTracker.h
@@ -104,6 +104,7 @@ public:
 	bool Destroy();
 	bool ConnectMRMesh(UMRMeshComponent* InMRMeshPtr);
 	bool DisconnectMRMesh(class UMRMeshComponent* InMRMeshPtr);
+	void SetUseWeightedNormals(const bool bInUseWeightedNormals);
 	bool Update();
 	bool RequestMesh();
 	bool GetMeshResult();
@@ -119,12 +120,11 @@ private:
 	MLHandle CurrentMeshRequest;
 #endif //WITH_MLSDK
 	bool bCreating;
+	bool bUseWeightedNormals;
 	// Maps system 128-bit mesh GUIDs to 64-bit brick indices
 	TMap<FGuid, uint64> MeshBrickCache;
 	uint64 MeshBrickIndex;
 	class UMRMeshComponent* MRMesh;
-	// A free list to recycle the CachedMeshData instances.
-	FMLCachedMeshData::SharedPtr CurrentMeshDataCache;
 	TArray<FMLCachedMeshData::SharedPtr> CachedMeshDatas;
 	TArray<FMLCachedMeshData::SharedPtr> FreeCachedMeshDatas;
 	FCriticalSection FreeCachedMeshDatasMutex; //The free list may be pushed/popped from multiple threads.
diff --git a/Engine/Plugins/Lumin/MagicLeap/Source/MagicLeapHandMeshing/Public/MagicLeapHandMeshingComponent.h b/Engine/Plugins/Lumin/MagicLeap/Source/MagicLeapHandMeshing/Public/MagicLeapHandMeshingComponent.h
--- a/Engine/Plugins/Lumin/MagicLeap/Source/MagicLeapHandMeshing/Public/MagicLeapHandMeshingComponent.h
+++ b/Engine/Plugins/Lumin/MagicLeap/Source/MagicLeapHandMeshing/Public/MagicLeapHandMeshingComponent.h
@@ -36,4 +36,11 @@ public:
 	*/
 	UFUNCTION(BlueprintCallable, Category = "HandMeshing Function Library | MagicLeap")
 	bool DisconnectMRMesh(class UMRMeshComponent* InMRMeshPtr);
+
+	/**
+		Specifies whether to use weighted triangle normals when calculating the normals for each vertex of the hand mesh.
+		@param bInUseWeightedNormals False is the default, True may impact performance but could give better looking results.
+	*/
+	UFUNCTION(BlueprintCallable, Category = "HandMeshing Function Library | MagicLeap")
+	void SetUseWeightedNormals(const bool bInUseWeightedNormals);
 };
