(Merge from MagicLeap-staging): Fix planes not working on consequetive VRPreview runs when using the AugmentedReality interface.

diff --git a/Engine/Plugins/Lumin/MagicLeap/Source/MagicLeapAR/Private/LuminARPlanesTracker.cpp b/Engine/Plugins/Lumin/MagicLeap/Source/MagicLeapAR/Private/LuminARPlanesTracker.cpp
--- a/Engine/Plugins/Lumin/MagicLeap/Source/MagicLeapAR/Private/LuminARPlanesTracker.cpp
+++ b/Engine/Plugins/Lumin/MagicLeap/Source/MagicLeapAR/Private/LuminARPlanesTracker.cpp
@@ -29,6 +29,8 @@ void FLuminARPlanesTracker::DestroyEntityTracker()
 	IMagicLeapPlanesModule::Get().RemoveQuery(PlanesQueryHandle);
 	IMagicLeapPlanesModule::Get().DestroyTracker();
 	PlaneResultsMap.Empty();
+	PlaneParentMap.Empty();
+	bPlanesQueryPending = false;
 }
 
 void FLuminARPlanesTracker::OnStartGameFrame()
diff --git a/Engine/Plugins/Lumin/MagicLeap/Source/MagicLeapAR/Private/LuminARPlanesTracker.h b/Engine/Plugins/Lumin/MagicLeap/Source/MagicLeapAR/Private/LuminARPlanesTracker.h
--- a/Engine/Plugins/Lumin/MagicLeap/Source/MagicLeapAR/Private/LuminARPlanesTracker.h
+++ b/Engine/Plugins/Lumin/MagicLeap/Source/MagicLeapAR/Private/LuminARPlanesTracker.h
@@ -47,7 +47,6 @@ private:
 	// Maps InnerID's to the plane and boundary data
 	TMap<FGuid, FLuminPlanesAndBoundaries> PlaneResultsMap;
 
-	int32 MaxPlaneQueryResults = 0;
 	bool bDiscardZeroExtentPlanes = false;
 
 	bool bPlanesQueryPending;
