Merge from partner-magicleap-staging:
Fixes in AugmentedReality impl for MagicLeap.

Reset worker thread setup for image tracking so it restarts on consequtive VRPreviews.
We need to UpdateTrackedGeometry in FLuminARTrackedPointResource::UpdateGeometryData even if we couldn't get a valid position and orientation to update to.  This can sometimes happen after a level transition, and if we don't UpdateTrackedGeometry, the Point may not have a valid ARSystem pointer, which may crash if certain blueprint functions (GetLocalToWorldTransform, for example) are called at this time.

Removed GetTrackableHandleMap from the ARTrackingSystem interface. Instead, a new function RemoveHandleIf is used to set the 'StoppedTracking' state and remove the handle.
RemoveHandleIf takes in a higher-order function that is applied to every stored handle if the geometry type matches. When the applied function returns true, the tracked handle stops tracking and is removed internally. Users of the tracking system can perform additional removal behavior as they see fit via the input function.

diff --git a/Engine/Plugins/Lumin/MagicLeap/Source/MagicLeap/Private/AppEventHandler.cpp b/Engine/Plugins/Lumin/MagicLeap/Source/MagicLeap/Private/AppEventHandler.cpp
--- a/Engine/Plugins/Lumin/MagicLeap/Source/MagicLeap/Private/AppEventHandler.cpp
+++ b/Engine/Plugins/Lumin/MagicLeap/Source/MagicLeap/Private/AppEventHandler.cpp
@@ -9,7 +9,8 @@
 namespace MagicLeap
 {
 	IAppEventHandler::IAppEventHandler(const TArray<EMagicLeapPrivilege>& InRequiredPrivilegeIDs)
-	: OnAppShutDownHandler(nullptr)
+	: OnAppStartHandler(nullptr)
+	, OnAppShutDownHandler(nullptr)
 	, OnAppTickHandler(nullptr)
 	, OnAppPauseHandler(nullptr)
 	, OnAppResumeHandler(nullptr)
@@ -146,6 +147,14 @@ namespace MagicLeap
 		return false;
 	}
 
+	void IAppEventHandler::OnAppStart()
+	{
+		if (OnAppStartHandler)
+		{
+			OnAppStartHandler();
+		}
+	}
+
 	void IAppEventHandler::OnAppShutDown()
 	{
 		if (OnAppShutDownHandler)
diff --git a/Engine/Plugins/Lumin/MagicLeap/Source/MagicLeap/Private/AppFramework.cpp b/Engine/Plugins/Lumin/MagicLeap/Source/MagicLeap/Private/AppFramework.cpp
--- a/Engine/Plugins/Lumin/MagicLeap/Source/MagicLeap/Private/AppFramework.cpp
+++ b/Engine/Plugins/Lumin/MagicLeap/Source/MagicLeap/Private/AppFramework.cpp
@@ -126,6 +126,15 @@ void FAppFramework::ApplicationResumeDelegate()
 	}
 }
 
+void FAppFramework::OnApplicationStart()
+{
+	FScopeLock Lock(&EventHandlersCriticalSection);
+	for (auto EventHandler : EventHandlers)
+	{
+		EventHandler->OnAppStart();
+	}
+}
+
 void FAppFramework::OnApplicationShutdown()
 {
 	FScopeLock Lock(&EventHandlersCriticalSection);
diff --git a/Engine/Plugins/Lumin/MagicLeap/Source/MagicLeap/Private/AppFramework.h b/Engine/Plugins/Lumin/MagicLeap/Source/MagicLeap/Private/AppFramework.h
--- a/Engine/Plugins/Lumin/MagicLeap/Source/MagicLeap/Private/AppFramework.h
+++ b/Engine/Plugins/Lumin/MagicLeap/Source/MagicLeap/Private/AppFramework.h
@@ -25,6 +25,7 @@ public:
 
 	void ApplicationPauseDelegate();
 	void ApplicationResumeDelegate();
+	void OnApplicationStart();
 	void OnApplicationShutdown();
 
 	void OnDeviceActive();
diff --git a/Engine/Plugins/Lumin/MagicLeap/Source/MagicLeap/Public/AppEventHandler.h b/Engine/Plugins/Lumin/MagicLeap/Source/MagicLeap/Public/AppEventHandler.h
--- a/Engine/Plugins/Lumin/MagicLeap/Source/MagicLeap/Public/AppEventHandler.h
+++ b/Engine/Plugins/Lumin/MagicLeap/Source/MagicLeap/Public/AppEventHandler.h
@@ -54,6 +54,11 @@ namespace MagicLeap
 		/** Removes the IAppEventHandler instance from the application's list of IAppEventHandler instances.*/
 		virtual ~IAppEventHandler();
 
+		/**
+			Perform any operations that must occur when an application begins
+		*/
+		virtual void OnAppStart();
+
 		/**
 			Can be overridden by inheriting class that needs to destroy certain api interfaces before the perception stack is
 			closed down.
@@ -101,6 +106,15 @@ namespace MagicLeap
 		*/
 		bool AddPrivilegeEventHandler(EMagicLeapPrivilege PrivilegeID, FRequiredPrivilege::FPrivilegeEventHandler&& InOnPrivilegeEvent);
 
+		/**
+			Use this as an alternative to overriding the OnAppStart function.  This allows you to use IAppEventHandler
+			as and aggregate class rather than an ancestor.
+		*/
+		void SetOnAppStartHandler(FEventHandler&& InOnAppStartHandler)
+		{
+			OnAppStartHandler = MoveTemp(InOnAppStartHandler);
+		}
+
 		/**
 			Use this as an alternative to overriding the OnAppShutDown function.  This allows you to use IAppEventHandler
 			as and aggregate class rather than an ancestor.
@@ -149,6 +163,7 @@ namespace MagicLeap
 
 	protected:
 		TMap<EMagicLeapPrivilege, FRequiredPrivilege> RequiredPrivileges;
+		FEventHandler OnAppStartHandler;
 		FEventHandler OnAppShutDownHandler;
 		FEventHandler OnAppTickHandler;
 		FEventHandler OnAppPauseHandler;
diff --git a/Engine/Plugins/Lumin/MagicLeap/Source/MagicLeap/Public/MagicLeapRunnable.h b/Engine/Plugins/Lumin/MagicLeap/Source/MagicLeap/Public/MagicLeapRunnable.h
--- a/Engine/Plugins/Lumin/MagicLeap/Source/MagicLeap/Public/MagicLeapRunnable.h
+++ b/Engine/Plugins/Lumin/MagicLeap/Source/MagicLeap/Public/MagicLeapRunnable.h
@@ -35,6 +35,7 @@ public:
 	, StopTaskCounter(0)
 	, Semaphore(nullptr)
 	, bPaused(false)
+	, bShuttingDown(false)
 	{
 		AppEventHandler.SetOnAppPauseHandler([this]() 
 		{
@@ -45,6 +46,16 @@ public:
 		{
 			OnAppResume();
 		});
+
+		AppEventHandler.SetOnAppStartHandler([this]()
+		{
+			OnAppStart();
+		});
+
+		AppEventHandler.SetOnAppShutDownHandler([this]()
+		{
+			OnAppShutdown();
+		});
 	}
 
 	virtual ~FMagicLeapRunnable()
@@ -118,8 +129,23 @@ public:
 		}
 	}
 
+	void OnAppStart()
+	{
+		bShuttingDown = false;
+	}
+
+	void OnAppShutdown()
+	{
+		bShuttingDown = true;
+	}
+
 	void PushNewTask(const TTaskType& InTask)
 	{
+		if (bShuttingDown)
+		{
+			return;
+		}
+
 		// on demand thread creation
 		if (!Thread)
 		{
@@ -175,6 +201,7 @@ protected:
 	FThreadSafeCounter StopTaskCounter;
 	FEvent* Semaphore;
 	FThreadSafeBool bPaused;
+	FThreadSafeBool bShuttingDown;
 	TQueue<TTaskType, EQueueMode::Spsc> IncomingTasks;
 	TQueue<TTaskType, EQueueMode::Mpsc> CompletedTasks; // allow multiple (binder) threads to push completed tasks
 	TTaskType CurrentTask;
diff --git a/Engine/Plugins/Lumin/MagicLeap/Source/MagicLeapAR/Private/LuminARPointsTracker.cpp b/Engine/Plugins/Lumin/MagicLeap/Source/MagicLeapAR/Private/LuminARPointsTracker.cpp
--- a/Engine/Plugins/Lumin/MagicLeap/Source/MagicLeapAR/Private/LuminARPointsTracker.cpp
+++ b/Engine/Plugins/Lumin/MagicLeap/Source/MagicLeapAR/Private/LuminARPointsTracker.cpp
@@ -80,16 +80,23 @@ void FLuminARTrackedPointResource::UpdateGeometryData(FLuminARImplementation* In
 	bool bPointFoundInEnvironment;
 	const bool bResult = UMagicLeapARPinFunctionLibrary::GetARPinPositionAndOrientation_TrackingSpace(GetNativeHandle(), Position, Orientation, bPointFoundInEnvironment);
 
+	const uint32 FrameNum = InARSystemSupport->GetFrameNum();
+	const int64 TimeStamp = InARSystemSupport->GetCameraTimestamp();
+	FTransform LocalToTrackingTransform;
+
 	if (!bResult)
 	{
-		// Something bad happend! don't update transform.
+		// Something bad happend!
 		// We can choose to mark this point's tracking state as "StoppedTracking" based on bPointFoundInEnvironment but that will happen automatically on next tick.
-		return;
-	}
-	const FTransform LocalToTrackingTransform(Orientation, Position);
 
-	const uint32 FrameNum = InARSystemSupport->GetFrameNum();
-	const int64 TimeStamp = InARSystemSupport->GetCameraTimestamp();
+		// In rare situations (usually at the end of a level transition), blueprint may call functions that require a valid ARSystem even though the above function
+		// may have returned false.  So we still need to update the transform, but we just update it to what it currently is.
+		LocalToTrackingTransform = Point->GetLocalToTrackingTransform();
+	}
+	else
+	{
+		LocalToTrackingTransform = FTransform(Orientation, Position);
+	}
 
 	Point->UpdateTrackedGeometry(InARSystemSupport->GetARSystem(), FrameNum, static_cast<double>(TimeStamp), LocalToTrackingTransform, InARSystemSupport->GetARSystem()->GetAlignmentTransform());
 	Point->SetDebugName(FName(*GetNativeHandle().ToString()));
diff --git a/Engine/Plugins/Lumin/MagicLeap/Source/MagicLeapAR/Private/LuminARTrackingSystem.cpp b/Engine/Plugins/Lumin/MagicLeap/Source/MagicLeapAR/Private/LuminARTrackingSystem.cpp
--- a/Engine/Plugins/Lumin/MagicLeap/Source/MagicLeapAR/Private/LuminARTrackingSystem.cpp
+++ b/Engine/Plugins/Lumin/MagicLeap/Source/MagicLeapAR/Private/LuminARTrackingSystem.cpp
@@ -276,6 +276,7 @@ void FLuminARImplementation::OnStartARSession(UARSessionConfig* SessionConfig)
 	}
 
 	bStartSessionRequested = true;
+	bUpdateTrackedImages = true;
 
 	// TODO : check if this code is needed.
 	// Try recreating the LuminARSession to fix the fatal error.
@@ -332,9 +333,15 @@ void FLuminARImplementation::OnStopARSession()
 		Tracker->DestroyEntityTracker();
 	}
 
+	// Delete the LightEstimateTracker on session end.  It will be recreated, if necessary, on session start.
+	Trackers.Remove(LightEstimateTracker);
+	delete LightEstimateTracker;
+	LightEstimateTracker = nullptr;
+
 	CurrentSessionStatus = EARSessionStatus::NotStarted;
 
 	TargetImageTextures.Empty();
+	TrackableHandleMap.Empty();
 }
 
 FARSessionStatus FLuminARImplementation::OnGetARSessionStatus() const
