Merge from MagicLeap-Staging:
Don't render the depth target to MLGraphics in ZI. Only reallocate depth target if requested size matches expected one. This workaround is required because this depth allocation function can receive requests for targets other than the one for eye stereo render (like when using a SceneCaptureComponent).

diff --git a/Engine/Plugins/Lumin/MagicLeap/Source/MagicLeap/Private/MagicLeapCustomPresentVk.cpp b/Engine/Plugins/Lumin/MagicLeap/Source/MagicLeap/Private/MagicLeapCustomPresentVk.cpp
--- a/Engine/Plugins/Lumin/MagicLeap/Source/MagicLeap/Private/MagicLeapCustomPresentVk.cpp
+++ b/Engine/Plugins/Lumin/MagicLeap/Source/MagicLeap/Private/MagicLeapCustomPresentVk.cpp
@@ -61,8 +61,11 @@ void FMagicLeapCustomPresentVulkan::FinishRendering()
 		const uint32 vp_width = static_cast<uint32>(FrameInfo.viewport.w);
 		const uint32 vp_height = static_cast<uint32>(FrameInfo.viewport.h);
 
+#if PLATFORM_LUMIN
+		// ZI doesn't make use of depth
 		// Clear MLGraphics depth textures since on Vulkan they are not pre-cleared. (needs to be done on RHI thread)
 		FMagicLeapHelperVulkan::ClearImage(FrameInfo.depth_id, FLinearColor::Black, 0 /* BaseMipLevel */, 1 /* LevelCount */, 0 /* BaseArrayLayer */, FrameInfo.num_virtual_cameras, true);
+#endif // PLATFORM_LUMIN
 
 		// Alias the render target with an srgb image description for proper color space output.
 		if (RenderTargetTextureAllocation != VK_NULL_HANDLE && LastAliasedRenderTarget != RenderTargetTexture)
@@ -73,7 +76,6 @@ void FMagicLeapCustomPresentVulkan::FinishRendering()
 				FMagicLeapHelperVulkan::DestroyImageSRGB(RenderTargetTextureSRGB);
 				RenderTargetTextureSRGB = VK_NULL_HANDLE;
 				RenderTargetTextureSRGB = reinterpret_cast<void *>(FMagicLeapHelperVulkan::AliasImageSRGB((uint64)RenderTargetTextureAllocation, RenderTargetTextureAllocationOffset, vp_width * 2, vp_height));
-				check(RenderTargetTextureSRGB != VK_NULL_HANDLE);
 			}
 			LastAliasedRenderTarget = RenderTargetTexture;
 			UE_LOG(LogMagicLeap, Log, TEXT("Aliased render target for correct sRGB ouput."));
@@ -83,10 +85,16 @@ void FMagicLeapCustomPresentVulkan::FinishRendering()
 		FMagicLeapHelperVulkan::BlitImage((uint64)FinalTarget, 0, 0, 0, vp_width, vp_height, 1, (uint64)FrameInfo.color_id, 0, 0, 0, 0, vp_width, vp_height, 1);
 		FMagicLeapHelperVulkan::BlitImage((uint64)FinalTarget, vp_width, 0, 0, vp_width, vp_height, 1, (uint64)FrameInfo.color_id, 1, 0, 0, 0, vp_width, vp_height, 1);
 
-		const VkImage DepthTarget = static_cast<FVulkanTexture2D*>(Plugin->DepthBuffer->GetTexture2D())->Surface.Image;
+#if PLATFORM_LUMIN
+		// ZI doesn't make use of depth
+		if (Plugin->DepthBuffer != nullptr)
+		{
+			const VkImage DepthTarget = static_cast<FVulkanTexture2D*>(Plugin->DepthBuffer->GetTexture2D())->Surface.Image;
 
-		FMagicLeapHelperVulkan::BlitImage((uint64)DepthTarget, 0, 0, 0, vp_width, vp_height, 1, (uint64)FrameInfo.depth_id, 0, 0, 0, 0, vp_width, vp_height, 1, true);
-		FMagicLeapHelperVulkan::BlitImage((uint64)DepthTarget, vp_width, 0, 0, vp_width, vp_height, 1, (uint64)FrameInfo.depth_id, 1, 0, 0, 0, vp_width, vp_height, 1, true);
+			FMagicLeapHelperVulkan::BlitImage((uint64)DepthTarget, 0, 0, 0, vp_width, vp_height, 1, (uint64)FrameInfo.depth_id, 0, 0, 0, 0, vp_width, vp_height, 1, true);
+			FMagicLeapHelperVulkan::BlitImage((uint64)DepthTarget, vp_width, 0, 0, vp_width, vp_height, 1, (uint64)FrameInfo.depth_id, 1, 0, 0, 0, vp_width, vp_height, 1, true);
+		}
+#endif // PLATFORM_LUMIN
 
 		FMagicLeapHelperVulkan::SignalObjects((uint64)FrameInfo.virtual_cameras[0].sync_object, (uint64)FrameInfo.virtual_cameras[1].sync_object, (uint64)FrameInfo.wait_sync_object);
 
diff --git a/Engine/Plugins/Lumin/MagicLeap/Source/MagicLeap/Private/MagicLeapHMD.cpp b/Engine/Plugins/Lumin/MagicLeap/Source/MagicLeap/Private/MagicLeapHMD.cpp
--- a/Engine/Plugins/Lumin/MagicLeap/Source/MagicLeap/Private/MagicLeapHMD.cpp
+++ b/Engine/Plugins/Lumin/MagicLeap/Source/MagicLeap/Private/MagicLeapHMD.cpp
@@ -539,7 +539,7 @@ bool FMagicLeapHMD::IsHMDConnected()
 		{
 			while (!IsEngineExitRequested())
 			{
-				bool bIsZIServerRunning;
+				bool bIsZIServerRunning = false;
 				MLResult Result = MLRemoteIsServerConfigured(&bIsZIServerRunning);
 				if (Result != MLResult_Ok && Result != MLResult_Timeout && Result != MLResult_NotImplemented)
 				{
@@ -744,6 +744,12 @@ void FMagicLeapHMD::RefreshTrackingFrame()
 #if WITH_MLSDK
 	check(IsInGameThread());
 
+	// Can't refresh tracking frame until head tracking has been enabled
+	if (HeadTracker == ML_INVALID_HANDLE)
+	{
+		return;
+	}
+
 	GameTrackingFrame.PixelDensity = PixelDensity;
 
 	// get the frame id for the frame
@@ -1191,19 +1197,26 @@ bool FMagicLeapHMD::NeedReAllocateDepthTexture(const TRefCountPtr<IPooledRenderT
 
 bool FMagicLeapHMD::AllocateDepthTexture(uint32 Index, uint32 SizeX, uint32 SizeY, uint8 Format, uint32 NumMips, uint32 FlagsIn, uint32 TargetableTextureFlags, FTexture2DRHIRef& OutTargetableTexture, FTexture2DRHIRef& OutShaderResourceTexture, uint32 NumSamples)
 {
-	// From AllocateRenderTargetTexture
-	if (!IsStereoEnabled())
+	if (IsStereoEnabled() && bNeedReAllocateDepthTexture)
 	{
-		return false;
-	}
+		FIntPoint ExpectedRenderTargetSize = GetIdealRenderTargetSize();
+		ExpectedRenderTargetSize.X = FMath::CeilToInt(ExpectedRenderTargetSize.X * PixelDensity);
+		ExpectedRenderTargetSize.Y = FMath::CeilToInt(ExpectedRenderTargetSize.Y * PixelDensity);
 
-	FRHIResourceCreateInfo CreateInfo(FClearValueBinding::DepthFar);
-	RHICreateTargetableShaderResource2D(SizeX, SizeY, PF_DepthStencil, 1, FlagsIn, TargetableTextureFlags | TexCreate_InputAttachmentRead, false, CreateInfo, OutTargetableTexture, OutShaderResourceTexture);
+		// Ensure the texture size matches eye render target size. We may get other depth allocations unrelated to the main scene render.
+		if (ExpectedRenderTargetSize == FIntPoint(SizeX, SizeY))
+		{
+			FRHIResourceCreateInfo CreateInfo(FClearValueBinding::DepthFar);
+			RHICreateTargetableShaderResource2D(SizeX, SizeY, PF_DepthStencil, 1, FlagsIn, TargetableTextureFlags | TexCreate_InputAttachmentRead, false, CreateInfo, OutTargetableTexture, OutShaderResourceTexture);
 
-	DepthBuffer = OutTargetableTexture;
-	bNeedReAllocateDepthTexture = false;
+			DepthBuffer = OutTargetableTexture;
+			bNeedReAllocateDepthTexture = false;
 
-	return true;
+			return true;
+		}
+	}
+
+	return false;
 }
 
 bool FMagicLeapHMD::GetHeadTrackingState(FMagicLeapHeadTrackingState& State) const
@@ -1301,7 +1314,10 @@ bool FMagicLeapHMD::AllocateRenderTargetTexture(uint32 Index, uint32 SizeX, uint
 	FRHIResourceCreateInfo CreateInfo;
 	RHICreateTargetableShaderResource2D(SizeX, SizeY, PF_R8G8B8A8, 1, TexCreate_None, TexCreate_RenderTargetable, false, CreateInfo, OutTargetableTexture, OutShaderResourceTexture);
 
+#if PLATFORM_LUMIN
+	// Depth should only be allocated if it is going to be used. Depth is not used with ZI
 	bNeedReAllocateDepthTexture = true;
+#endif // PLATFORM_LUMIN
 
 	return true;
 }
@@ -1544,6 +1560,7 @@ void FMagicLeapHMD::EnableDeviceFeatures()
 		// we get the profile effects on non-vr-preview rendering.
 		EnableLuminProfile();
 	}
+	AppFramework.OnApplicationStart();
 }
 
 void FMagicLeapHMD::DisableDeviceFeatures()
@@ -1580,6 +1597,13 @@ void FMagicLeapHMD::DisableDeviceFeatures()
 void FMagicLeapHMD::InitDevice_RenderThread()
 {
 #if WITH_MLSDK
+	// Can't initialize render thread until head tracking has been enabled
+	MLHandle HeadTrackerHandle = FPlatformAtomics::AtomicRead(reinterpret_cast<const volatile int64*>(&HeadTracker));
+	if (HeadTrackerHandle == ML_INVALID_HANDLE)
+	{
+		return;
+	}
+
 	if (bQueuedGraphicsCreateCall)
 	{
 		UE_LOG(LogMagicLeap, Warning, TEXT("Graphics client create call already queued."));
@@ -1591,11 +1615,16 @@ void FMagicLeapHMD::InitDevice_RenderThread()
 		bool bDeviceSuccessfullyInitialized = false;
 		// Unreal supports sRGB which is the default we are requesting from graphics as well now.
 		MLGraphicsOptions gfx_opts;
-		gfx_opts.graphics_flags = 0; //MLGraphicsFlags_DebugMode;
 		gfx_opts.color_format = MLSurfaceFormat_RGBA8UNormSRGB;
-		gfx_opts.depth_format = MLSurfaceFormat_D32Float;
 
+#if PLATFORM_LUMIN
 		gfx_opts.graphics_flags = MLGraphicsFlags_Default;
+		gfx_opts.depth_format = MLSurfaceFormat_D32Float;
+#else
+		// ZI doesn't make use of depth
+		gfx_opts.graphics_flags = MLGraphicsFlags_NoDepth;
+#endif
+
 
 #if PLATFORM_WINDOWS
 		if (IsPCPlatform(GMaxRHIShaderPlatform) && !IsOpenGLPlatform(GMaxRHIShaderPlatform))
