Fixing some fallout from some API deprecations.

diff --git a/Engine/Plugins/Lumin/MagicLeap/Source/MagicLeap/Private/MagicLeapHMD.cpp b/Engine/Plugins/Lumin/MagicLeap/Source/MagicLeap/Private/MagicLeapHMD.cpp
--- a/Engine/Plugins/Lumin/MagicLeap/Source/MagicLeap/Private/MagicLeapHMD.cpp
+++ b/Engine/Plugins/Lumin/MagicLeap/Source/MagicLeap/Private/MagicLeapHMD.cpp
@@ -1124,7 +1124,7 @@ bool FMagicLeapHMD::SetIgnoreInput(bool Ignore)
 				// First time around we call this to disable the input globally. Hence we
 				// also set mouse options. On subsequent calls we only set the input ignore flags.
 				DisableInputForBeginPlay = false;
-				ViewportClient->SetCaptureMouseOnClick(EMouseCaptureMode::NoCapture);
+				ViewportClient->SetMouseCaptureMode(EMouseCaptureMode::NoCapture);
 				ViewportClient->SetMouseLockMode(EMouseLockMode::DoNotLock);
 				ViewportClient->SetHideCursorDuringCapture(false);
 			}
