Add missing resource transition to Lumin and remove unnecessary assert

diff --git a/Engine/Plugins/Lumin/MagicLeap/Source/MagicLeap/Private/MagicLeapHMD.cpp b/Engine/Plugins/Lumin/MagicLeap/Source/MagicLeap/Private/MagicLeapHMD.cpp
--- a/Engine/Plugins/Lumin/MagicLeap/Source/MagicLeap/Private/MagicLeapHMD.cpp
+++ b/Engine/Plugins/Lumin/MagicLeap/Source/MagicLeap/Private/MagicLeapHMD.cpp
@@ -2028,6 +2028,8 @@ void FMagicLeapHMD::RenderTexture_RenderThread(class FRHICommandListImmediate& R
 		const uint32 TextureWidth = SrcTexture->GetSizeX();
 		const uint32 TextureHeight = SrcTexture->GetSizeY();
 
+		RHICmdList.TransitionResource(EResourceTransitionAccess::EReadable, SrcTexture);
+
 		// The BackBuffer is the debug view for mirror modes, i.e. vr-preview. In which case
 		// it can be an arbitrary size different than the render size. Which means
 		// we scale to that BackBuffer size, with either a letter-box or pill-box to
diff --git a/Engine/Plugins/Lumin/MagicLeap/Source/MagicLeap/Private/MagicLeapHMD.h b/Engine/Plugins/Lumin/MagicLeap/Source/MagicLeap/Private/MagicLeapHMD.h
--- a/Engine/Plugins/Lumin/MagicLeap/Source/MagicLeap/Private/MagicLeapHMD.h
+++ b/Engine/Plugins/Lumin/MagicLeap/Source/MagicLeap/Private/MagicLeapHMD.h
@@ -109,7 +109,6 @@ public:
 	virtual void UpdateViewportRHIBridge(bool bUseSeparateRenderTarget, const class FViewport& Viewport, FRHIViewport* const ViewportRHI) override;
 	virtual bool ShouldUseSeparateRenderTarget() const override
 	{
-		check(IsInGameThread());
 		return IsStereoEnabled();
 	}
 	virtual bool NeedReAllocateDepthTexture(const TRefCountPtr<IPooledRenderTarget>& DepthTarget) override;
