Merge from Partner-MagicLeap-Staging via Dev-VR:
Add missing BP redirects for image tracker module.

diff --git a/Engine/Plugins/Lumin/MagicLeap/Config/BaseMagicLeap.ini b/Engine/Plugins/Lumin/MagicLeap/Config/BaseMagicLeap.ini
--- a/Engine/Plugins/Lumin/MagicLeap/Config/BaseMagicLeap.ini
+++ b/Engine/Plugins/Lumin/MagicLeap/Config/BaseMagicLeap.ini
@@ -85,3 +85,13 @@
 ; ImageTracking redirects
 +ClassRedirects=(OldName="/Script/MagicLeap.ImageTrackerComponent", NewName="/Script/MagicLeapImageTracker.MagicLeapImageTrackerComponent")
 +ClassRedirects=(OldName="/Script/MagicLeap.ImageTrackerFunctionLibrary", NewName="/Script/MagicLeapImageTracker.MagicLeapImageTrackerFunctionLibrary")
++ClassRedirects=(OldName="/Script/MagicLeap.FImageTargetReliableTracking", NewName="/Script/MagicLeapImageTracker.FMagicLeapImageTargetReliableTracking")
++ClassRedirects=(OldName="/Script/MagicLeap.FImageTargetReliableTrackingMulti", NewName="/Script/MagicLeapImageTracker.FMagicLeapImageTargetReliableTrackingMulti")
++ClassRedirects=(OldName="/Script/MagicLeap.FSetImageTargetSucceeded", NewName="/Script/MagicLeapImageTracker.FMagicLeapSetImageTargetSucceeded")
++ClassRedirects=(OldName="/Script/MagicLeap.FSetImageTargetSucceededMulti", NewName="/Script/MagicLeapImageTracker.FMagicLeapSetImageTargetSucceededMulti")
++ClassRedirects=(OldName="/Script/MagicLeap.FSetImageTargetFailed", NewName="/Script/MagicLeapImageTracker.FMagicLeapSetImageTargetFailed")
++ClassRedirects=(OldName="/Script/MagicLeap.FSetImageTargetFailedMulti", NewName="/Script/MagicLeapImageTracker.FMagicLeapSetImageTargetFailedMulti")
++ClassRedirects=(OldName="/Script/MagicLeap.FImageTargetFound", NewName="/Script/MagicLeapImageTracker.FMagicLeapImageTargetFound")
++ClassRedirects=(OldName="/Script/MagicLeap.FImageTargetFoundMulti", NewName="/Script/MagicLeapImageTracker.FMagicLeapImageTargetFoundMulti")
++ClassRedirects=(OldName="/Script/MagicLeap.FImageTargetLost", NewName="/Script/MagicLeapImageTracker.FMagicLeapImageTargetLost")
++ClassRedirects=(OldName="/Script/MagicLeap.FImageTargetLostMulti", NewName="/Script/MagicLeapImageTracker.FMagicLeapImageTargetLostMulti")
