AudioCapture component update for lumin.

diff --git a/Engine/Plugins/Lumin/MagicLeap/Source/MagicLeapAudioCapture/Private/AudioCaptureMagicLeap.cpp b/Engine/Plugins/Lumin/MagicLeap/Source/MagicLeapAudioCapture/Private/AudioCaptureMagicLeap.cpp
--- a/Engine/Plugins/Lumin/MagicLeap/Source/MagicLeapAudioCapture/Private/AudioCaptureMagicLeap.cpp
+++ b/Engine/Plugins/Lumin/MagicLeap/Source/MagicLeapAudioCapture/Private/AudioCaptureMagicLeap.cpp
@@ -193,7 +193,7 @@ void Audio::FAudioCaptureMagicLeapStream::OnAudioCapture(void* InBuffer, uint32
 		FloatBufferPtr[i] = (((float)InBufferData[i]) / 32767.0f);
 	};
 
-	OnCapture(FloatBufferPtr, InBufferFrames, NumChannels, StreamTime, bOverflow);
+	OnCapture(FloatBufferPtr, InBufferFrames, NumChannels, SampleRate, StreamTime, bOverflow);
 }
 
 bool Audio::FAudioCaptureMagicLeapStream::GetInputDevicesAvailable(TArray<FCaptureDeviceInfo>& OutDevices)
