Fix cook warning about uninitialized Error field.

#rnx

diff --git a/Engine/Plugins/Lumin/MagicLeap/Source/MagicLeap/Public/MagicLeapHMDFunctionLibrary.h b/Engine/Plugins/Lumin/MagicLeap/Source/MagicLeap/Public/MagicLeapHMDFunctionLibrary.h
--- a/Engine/Plugins/Lumin/MagicLeap/Source/MagicLeap/Public/MagicLeapHMDFunctionLibrary.h
+++ b/Engine/Plugins/Lumin/MagicLeap/Source/MagicLeap/Public/MagicLeapHMDFunctionLibrary.h
@@ -44,6 +44,13 @@ struct MAGICLEAP_API FMagicLeapHeadTrackingState
 	GENERATED_BODY()
 
 public:
+	FMagicLeapHeadTrackingState()
+	{
+		Mode = EMagicLeapHeadTrackingMode::Unknown;
+		Error = EMagicLeapHeadTrackingError::None;
+		Confidence = 1.0f;
+	}
+
 	UPROPERTY(EditAnywhere, BlueprintReadWrite, Category = "MagicLeap")
 	EMagicLeapHeadTrackingMode Mode;
 
