Fix CIS warning for duplicate exported script name (Magic Leap).

diff --git a/Engine/Plugins/Lumin/MagicLeap/Source/MagicLeap/Public/MagicLeapMeshTrackerComponent.h b/Engine/Plugins/Lumin/MagicLeap/Source/MagicLeap/Public/MagicLeapMeshTrackerComponent.h
--- a/Engine/Plugins/Lumin/MagicLeap/Source/MagicLeap/Public/MagicLeapMeshTrackerComponent.h
+++ b/Engine/Plugins/Lumin/MagicLeap/Source/MagicLeap/Public/MagicLeapMeshTrackerComponent.h
@@ -59,7 +59,7 @@ public:
 	int32 GetNumQueuedBlockUpdates();
 
 	/** IMagicLeapMeshBlockSelectorInterface */
-	UFUNCTION(BlueprintCallable, BlueprintNativeEvent, Category="Meshing|MagicLeap")
+	UFUNCTION(BlueprintCallable, BlueprintNativeEvent, Category="Meshing|MagicLeap", meta = (ScriptName = "MeshTrackerSelectMeshBlocks", DisplayName = "Mesh Tracker Select Mesh Blocks"))
 	void SelectMeshBlocks(const FMagicLeapTrackingMeshInfo& NewMeshInfo, TArray<FMagicLeapMeshBlockRequest>& RequestedMesh);
 	virtual void SelectMeshBlocks_Implementation(const FMagicLeapTrackingMeshInfo& NewMeshInfo, TArray<FMagicLeapMeshBlockRequest>& RequestedMesh) override;
 
diff --git a/Engine/Plugins/Lumin/MagicLeap/Source/MagicLeap/Public/MeshBlockSelectorInterface.h b/Engine/Plugins/Lumin/MagicLeap/Source/MagicLeap/Public/MeshBlockSelectorInterface.h
--- a/Engine/Plugins/Lumin/MagicLeap/Source/MagicLeap/Public/MeshBlockSelectorInterface.h
+++ b/Engine/Plugins/Lumin/MagicLeap/Source/MagicLeap/Public/MeshBlockSelectorInterface.h
@@ -22,6 +22,6 @@ public:
 	 * @param NewMeshInfo Information on the latest mesh blocks available.
 	 * @param RequestedMesh output array containing details of blocks for whom the mesh should be requested.
 	 */
-	UFUNCTION(BlueprintCallable, BlueprintNativeEvent, Category="Meshing|MagicLeap")
+	UFUNCTION(BlueprintCallable, BlueprintNativeEvent, Category="Meshing|MagicLeap", meta = (ScriptName = "MeshBlockSelectorSelectMeshBlocks", DisplayName = "Mesh Block Selector Select Mesh Blocks"))
 	void SelectMeshBlocks(const FMagicLeapTrackingMeshInfo& NewMeshInfo, TArray<FMagicLeapMeshBlockRequest>& RequestedMesh);
 };
